Warning fix.

diff --git a/Engine/Plugins/Media/BlackmagicMedia/Source/BlackmagicMedia/Private/Blackmagic/Blackmagic.cpp b/Engine/Plugins/Media/BlackmagicMedia/Source/BlackmagicMedia/Private/Blackmagic/Blackmagic.cpp
--- a/Engine/Plugins/Media/BlackmagicMedia/Source/BlackmagicMedia/Private/Blackmagic/Blackmagic.cpp
+++ b/Engine/Plugins/Media/BlackmagicMedia/Source/BlackmagicMedia/Private/Blackmagic/Blackmagic.cpp
@@ -78,7 +78,7 @@ void FBlackmagic::LogInfo(const TCHAR* InFormat, ...)
 	va_list Args;
 
 	va_start(Args, InFormat );
-	FCString::GetVarArgs(TempString, ARRAY_COUNT(TempString), ARRAY_COUNT(TempString) - 1, InFormat, Args);
+	FCString::GetVarArgs(TempString, ARRAY_COUNT(TempString), InFormat, Args);
 	va_end(Args);
 
 	UE_LOG(LogBlackmagicMedia, Log, TempString);
@@ -92,7 +92,7 @@ void FBlackmagic::LogWarning(const TCHAR* InFormat, ...)
 	va_list Args;
 
 	va_start(Args, InFormat );
-	FCString::GetVarArgs(TempString, ARRAY_COUNT(TempString), ARRAY_COUNT(TempString) - 1, InFormat, Args);
+	FCString::GetVarArgs(TempString, ARRAY_COUNT(TempString), InFormat, Args);
 	va_end(Args);
 
 	UE_LOG(LogBlackmagicMedia, Warning, TempString);
@@ -106,7 +106,7 @@ void FBlackmagic::LogError(const TCHAR* InFormat, ...)
 	va_list Args;
 
 	va_start(Args, InFormat );
-	FCString::GetVarArgs(TempString, ARRAY_COUNT(TempString), ARRAY_COUNT(TempString) - 1, InFormat, Args);
+	FCString::GetVarArgs(TempString, ARRAY_COUNT(TempString), InFormat, Args);
 	va_end(Args);
 
 	UE_LOG(LogBlackmagicMedia, Error, TempString);
