Schedule audio samples in different thread than video in blackmagic

diff --git a/Engine/Plugins/Media/BlackmagicMedia/Source/BlackmagicMediaOutput/Private/BlackmagicMediaCapture.cpp b/Engine/Plugins/Media/BlackmagicMedia/Source/BlackmagicMediaOutput/Private/BlackmagicMediaCapture.cpp
--- a/Engine/Plugins/Media/BlackmagicMedia/Source/BlackmagicMediaOutput/Private/BlackmagicMediaCapture.cpp
+++ b/Engine/Plugins/Media/BlackmagicMedia/Source/BlackmagicMediaOutput/Private/BlackmagicMediaCapture.cpp
@@ -21,6 +21,8 @@
 #include "EngineAnalytics.h"
 #endif
 
+#include <iostream>
+#include <fstream>
 
 bool bBlackmagicWritInputRawDataCmdEnable = false;
 static FAutoConsoleCommand BlackmagicWriteInputRawDataCmd(
@@ -535,6 +537,7 @@ bool UBlackmagicMediaCapture::InitBlackmagic(UBlackmagicMediaOutput* InBlackmagi
 void UBlackmagicMediaCapture::OnFrameCaptured_RenderingThread(const FCaptureBaseData& InBaseData, TSharedPtr<FMediaCaptureUserData, ESPMode::ThreadSafe> InUserData, void* InBuffer, int32 Width, int32 Height, int32 BytesPerRow)
 {
 	// Prevent the rendering thread from copying while we are stopping the capture.
+	TRACE_CPUPROFILER_EVENT_SCOPE(UBlackmagicMediaCapture::OnFrameCaptured_RenderingThread);
 	FScopeLock ScopeLock(&RenderThreadCriticalSection);
 	if (EventCallback)
 	{
@@ -551,16 +554,23 @@ void UBlackmagicMediaCapture::OnFrameCaptured_RenderingThread(const FCaptureBase
 		Frame.VideoHeight = Height;
 		Frame.Timecode = Timecode;
 		Frame.FrameIdentifier = InBaseData.SourceFrameNumber;
+
+		bool bSent = false;
+		{
+			TRACE_CPUPROFILER_EVENT_SCOPE(UBlackmagicMediaCapture::SendVideoFrameData);
+			bSent = EventCallback->SendVideoFrameData(Frame);
+		}
 		
-		const bool bSent = EventCallback->SendVideoFrameData(Frame);
 		
 		if (bLogDropFrame && !bSent)
 		{
 			UE_LOG(LogBlackmagicMediaOutput, Warning, TEXT("Frame couldn't be sent to Blackmagic device. Engine might be running faster than output."));
 		}
 
-		if (bSent && bOutputAudio)
+		if (bOutputAudio/* && bSent*/)
 		{
+			TRACE_CPUPROFILER_EVENT_SCOPE(UBlackmagicMediaCapture::OutputAudio);
+
 			BlackmagicDesign::FAudioSamplesDescriptor AudioSamples;
 			AudioSamples.Timecode = Timecode;
 			AudioSamples.FrameIdentifier = InBaseData.SourceFrameNumber;
@@ -582,6 +592,7 @@ void UBlackmagicMediaCapture::OnFrameCaptured_RenderingThread(const FCaptureBase
 					AudioSamples.AudioBuffer = reinterpret_cast<uint8_t*>(AudioBuffer.GetData());
 					AudioSamples.NumAudioSamples = AudioBuffer.Num() / NumOutputChannels;
 					AudioSamples.AudioBufferLength = AudioBuffer.Num() * sizeof(int16);
+
 					EventCallback->SendAudioSamples(AudioSamples);
 				}
 				else
diff --git a/Engine/Plugins/Media/BlackmagicMedia/Source/BlackmagicMediaOutput/Public/BlackmagicMediaCapture.h b/Engine/Plugins/Media/BlackmagicMedia/Source/BlackmagicMediaOutput/Public/BlackmagicMediaCapture.h
--- a/Engine/Plugins/Media/BlackmagicMedia/Source/BlackmagicMediaOutput/Public/BlackmagicMediaCapture.h
+++ b/Engine/Plugins/Media/BlackmagicMedia/Source/BlackmagicMediaOutput/Public/BlackmagicMediaCapture.h
@@ -49,6 +49,7 @@ private:
 	void WaitForSync_RenderingThread();
 	void ApplyViewportTextureAlpha(TSharedPtr<FSceneViewport> InSceneViewport);
 	void RestoreViewportTextureAlpha(TSharedPtr<FSceneViewport> InSceneViewport);
+	void TriggerAudioRender();
 
 private:
 	friend BlackmagicMediaCaptureHelpers::FBlackmagicMediaCaptureEventCallback;
