- Fixing cvar name clash in media io

diff --git a/Engine/Plugins/Media/AjaMedia/Source/AjaMediaOutput/Private/AjaMediaCapture.cpp b/Engine/Plugins/Media/AjaMedia/Source/AjaMediaOutput/Private/AjaMediaCapture.cpp
--- a/Engine/Plugins/Media/AjaMedia/Source/AjaMediaOutput/Private/AjaMediaCapture.cpp
+++ b/Engine/Plugins/Media/AjaMedia/Source/AjaMediaOutput/Private/AjaMediaCapture.cpp
@@ -21,7 +21,7 @@
 #include "EngineAnalytics.h"
 #endif
 
-TAutoConsoleVariable<int32> CVarEnableGPUDirect(
+TAutoConsoleVariable<int32> CVarAjaEnableGPUDirect(
 	TEXT("Aja.EnableGPUDirect"), 0,
 	TEXT("Whether to enable GPU direct for faster video frame copies. (Experimental)"),
 	ECVF_RenderThreadSafe);
@@ -322,7 +322,7 @@ void UAjaMediaCapture::StopCaptureImpl(bool bAllowPendingFrameToBeProcess)
 
 bool UAjaMediaCapture::ShouldCaptureRHITexture() const
 {
-	return bGPUTextureTransferAvailable && CVarEnableGPUDirect.GetValueOnAnyThread() == 1;
+	return bGPUTextureTransferAvailable && CVarAjaEnableGPUDirect.GetValueOnAnyThread() == 1;
 }
 
 void UAjaMediaCapture::BeforeFrameCaptured_RenderingThread(const FCaptureBaseData& InBaseData, TSharedPtr<FMediaCaptureUserData, ESPMode::ThreadSafe> InUserData, FTextureRHIRef InTexture)
diff --git a/Engine/Plugins/Media/BlackmagicMedia/Source/BlackmagicMediaOutput/Private/BlackmagicMediaCapture.cpp b/Engine/Plugins/Media/BlackmagicMedia/Source/BlackmagicMediaOutput/Private/BlackmagicMediaCapture.cpp
--- a/Engine/Plugins/Media/BlackmagicMedia/Source/BlackmagicMediaOutput/Private/BlackmagicMediaCapture.cpp
+++ b/Engine/Plugins/Media/BlackmagicMedia/Source/BlackmagicMediaOutput/Private/BlackmagicMediaCapture.cpp
@@ -21,7 +21,7 @@
 #include "EngineAnalytics.h"
 #endif
 
-TAutoConsoleVariable<int32> CVarEnableGPUDirect(
+TAutoConsoleVariable<int32> CVarBlackmagicEnableGPUDirect(
 	TEXT("Blackmagic.EnableGPUDirect"), 0,
 	TEXT("Whether to enable GPU direct for faster video frame copies. (Experimental)"),
 	ECVF_RenderThreadSafe);
@@ -426,7 +426,7 @@ void UBlackmagicMediaCapture::StopCaptureImpl(bool bAllowPendingFrameToBeProcess
 bool UBlackmagicMediaCapture::ShouldCaptureRHITexture() const
 {
 	// Todo: also test if dvp was initialized correctly.
-	return bGPUTextureTransferAvailable && CVarEnableGPUDirect.GetValueOnAnyThread() == 1;
+	return bGPUTextureTransferAvailable && CVarBlackmagicEnableGPUDirect.GetValueOnAnyThread() == 1;
 }
 
 void UBlackmagicMediaCapture::ApplyViewportTextureAlpha(TSharedPtr<FSceneViewport> InSceneViewport)
