Fix com warning by not initializing blackmagic when the card can't be used

diff --git a/Engine/Plugins/Media/BlackmagicMedia/Source/BlackmagicMedia/Private/Assets/BlackmagicCustomTimeStep.cpp b/Engine/Plugins/Media/BlackmagicMedia/Source/BlackmagicMedia/Private/Assets/BlackmagicCustomTimeStep.cpp
--- a/Engine/Plugins/Media/BlackmagicMedia/Source/BlackmagicMedia/Private/Assets/BlackmagicCustomTimeStep.cpp
+++ b/Engine/Plugins/Media/BlackmagicMedia/Source/BlackmagicMedia/Private/Assets/BlackmagicCustomTimeStep.cpp
@@ -4,6 +4,7 @@
 
 #include "Blackmagic.h"
 #include "BlackmagicMediaPrivate.h"
+#include "IBlackmagicMediaModule.h"
 
 #include "HAL/Event.h"
 #include "HAL/IConsoleManager.h"
@@ -232,19 +233,13 @@ bool UBlackmagicCustomTimeStep::Initialize(class UEngine* InEngine)
 	}
 
 	check(InputEventCallback == nullptr);
-
-	if (!FBlackmagic::IsInitialized())
+	
+	if (!IBlackmagicMediaModule::Get().IsInitialized())
 	{
 		UE_LOG(LogBlackmagicMedia, Error, TEXT("The CustomTimeStep '%s' can't be initialized. Blackmagic is not initialized on your machine."), *GetName());
 		return false;
 	}
 
-	if (!FBlackmagic::CanUseBlackmagicCard())
-	{
-		UE_LOG(LogBlackmagicMedia, Warning, TEXT("The CustomTimeStep '%s' can't be initialized because Blackmagic card cannot be used. Are you in a Commandlet? You may override this behavior by launching with -ForceBlackmagicUsage"), *GetName());
-		return false;
-	}
-
 	BlackmagicDesign::FChannelInfo ChannelInfo;
 	ChannelInfo.DeviceIndex = MediaConfiguration.MediaConnection.Device.DeviceIdentifier;
 
diff --git a/Engine/Plugins/Media/BlackmagicMedia/Source/BlackmagicMedia/Private/Assets/BlackmagicMediaSource.cpp b/Engine/Plugins/Media/BlackmagicMedia/Source/BlackmagicMedia/Private/Assets/BlackmagicMediaSource.cpp
--- a/Engine/Plugins/Media/BlackmagicMedia/Source/BlackmagicMedia/Private/Assets/BlackmagicMediaSource.cpp
+++ b/Engine/Plugins/Media/BlackmagicMedia/Source/BlackmagicMedia/Private/Assets/BlackmagicMediaSource.cpp
@@ -4,6 +4,7 @@
 
 #include "Blackmagic.h"
 #include "BlackmagicMediaPrivate.h"
+#include "IBlackmagicMediaModule.h"
 
 #include "MediaIOCorePlayerBase.h"
 
@@ -116,17 +117,11 @@ bool UBlackmagicMediaSource::Validate() const
 		return false;
 	}
 
-	if (!FBlackmagic::IsInitialized())
+	if (!IBlackmagicMediaModule::Get().IsInitialized())
 	{
 		UE_LOG(LogBlackmagicMedia, Warning, TEXT("Can't validate MediaSource '%s'. the Blackmagic library was not initialized."), *GetName());
 		return false;
 	}
-
-	if (!FBlackmagic::CanUseBlackmagicCard())
-	{
-		UE_LOG(LogBlackmagicMedia, Warning, TEXT("Can't validate MediaSource '%s' because Blackmagic card cannot be used. Are you in a Commandlet? You may override this behavior by launching with -ForceBlackmagicUsage"), *GetName());
-		return false;
-	}
 	
 	TUniquePtr<BlackmagicDesign::BlackmagicDeviceScanner> Scanner = MakeUnique<BlackmagicDesign::BlackmagicDeviceScanner>();
 	BlackmagicDesign::BlackmagicDeviceScanner::DeviceInfo DeviceInfo;
diff --git a/Engine/Plugins/Media/BlackmagicMedia/Source/BlackmagicMedia/Private/Assets/BlackmagicTimecodeProvider.cpp b/Engine/Plugins/Media/BlackmagicMedia/Source/BlackmagicMedia/Private/Assets/BlackmagicTimecodeProvider.cpp
--- a/Engine/Plugins/Media/BlackmagicMedia/Source/BlackmagicMedia/Private/Assets/BlackmagicTimecodeProvider.cpp
+++ b/Engine/Plugins/Media/BlackmagicMedia/Source/BlackmagicMedia/Private/Assets/BlackmagicTimecodeProvider.cpp
@@ -3,6 +3,7 @@
 #include "BlackmagicTimecodeProvider.h"
 #include "BlackmagicMediaPrivate.h"
 #include "Blackmagic.h"
+#include "IBlackmagicMediaModule.h"
 
 #include "HAL/CriticalSection.h"
 #include "Misc/ScopeLock.h"
@@ -167,18 +168,12 @@ bool UBlackmagicTimecodeProvider::Initialize(class UEngine* InEngine)
 		return false;
 	}
 
-	if (!FBlackmagic::IsInitialized())
+	if (!IBlackmagicMediaModule::Get().IsInitialized())
 	{
 		UE_LOG(LogBlackmagicMedia, Error, TEXT("The TimecodeProvider '%s' can't be initialized. Blackmagic is not initialized on your machine."), *GetName());
 		return false;
 	}
 
-	if (!FBlackmagic::CanUseBlackmagicCard())
-	{
-		UE_LOG(LogBlackmagicMedia, Warning, TEXT("The TimecodeProvider '%s' can't be initialized because Blackmagic card cannot be used. Are you in a Commandlet? You may override this behavior by launching with -ForceBlackmagicUsage"), *GetName());
-		return false;
-	}
-
 	if (TimecodeFormat == EMediaIOTimecodeFormat::None)
 	{
 		UE_LOG(LogBlackmagicMedia, Error, TEXT("The TimecodeProvider '%s' can't be initialized. Selected timecode format is invalid."), *GetName());
diff --git a/Engine/Plugins/Media/BlackmagicMedia/Source/BlackmagicMedia/Private/Blackmagic/BlackmagicDeviceProvider.cpp b/Engine/Plugins/Media/BlackmagicMedia/Source/BlackmagicMedia/Private/Blackmagic/BlackmagicDeviceProvider.cpp
--- a/Engine/Plugins/Media/BlackmagicMedia/Source/BlackmagicMedia/Private/Blackmagic/BlackmagicDeviceProvider.cpp
+++ b/Engine/Plugins/Media/BlackmagicMedia/Source/BlackmagicMedia/Private/Blackmagic/BlackmagicDeviceProvider.cpp
@@ -5,6 +5,7 @@
 #include "Blackmagic.h"
 #include "BlackmagicLib.h"
 #include "BlackmagicMediaPrivate.h"
+#include "IBlackmagicMediaModule.h"
 #include "MediaIOCoreCommonDisplayMode.h"
 
 #define LOCTEXT_NAMESPACE "BlackmagicDeviceProvider"
@@ -106,7 +107,8 @@ TArray<FMediaIOConfiguration> FBlackmagicDeviceProvider::GetConfigurations(bool
 	const int32 MaxNumberOfChannel = 8;
 
 	TArray<FMediaIOConfiguration> Results;
-	if (!FBlackmagic::IsInitialized() || !FBlackmagic::CanUseBlackmagicCard())
+	
+	if (!IBlackmagicMediaModule::Get().IsInitialized())
 	{
 		return Results;
 	}
@@ -281,7 +283,7 @@ TArray<FMediaIOOutputConfiguration> FBlackmagicDeviceProvider::GetOutputConfigur
 TArray<FMediaIODevice> FBlackmagicDeviceProvider::GetDevices() const
 {
 	TArray<FMediaIODevice> Results;
-	if (!FBlackmagic::IsInitialized() || !FBlackmagic::CanUseBlackmagicCard())
+	if (!IBlackmagicMediaModule::Get().IsInitialized())
 	{
 		return Results;
 	}
diff --git a/Engine/Plugins/Media/BlackmagicMedia/Source/BlackmagicMedia/Private/BlackmagicMediaModule.cpp b/Engine/Plugins/Media/BlackmagicMedia/Source/BlackmagicMedia/Private/BlackmagicMediaModule.cpp
--- a/Engine/Plugins/Media/BlackmagicMedia/Source/BlackmagicMedia/Private/BlackmagicMediaModule.cpp
+++ b/Engine/Plugins/Media/BlackmagicMedia/Source/BlackmagicMedia/Private/BlackmagicMediaModule.cpp
@@ -50,10 +50,13 @@ public:
 	virtual void StartupModule() override
 	{
 		// initialize
-		if (!FBlackmagic::Initialize())
+		if (CanBeUsed())
 		{
-			UE_LOG(LogBlackmagicMedia, Error, TEXT("Failed to initialize Blackmagic"));
-			return;
+			if (!FBlackmagic::Initialize())
+			{
+				UE_LOG(LogBlackmagicMedia, Error, TEXT("Failed to initialize Blackmagic"));
+				return;
+			}
 		}
 
 		CreateStyle();
diff --git a/Engine/Plugins/Media/BlackmagicMedia/Source/BlackmagicMedia/Private/Player/BlackmagicMediaPlayer.cpp b/Engine/Plugins/Media/BlackmagicMedia/Source/BlackmagicMedia/Private/Player/BlackmagicMediaPlayer.cpp
--- a/Engine/Plugins/Media/BlackmagicMedia/Source/BlackmagicMedia/Private/Player/BlackmagicMediaPlayer.cpp
+++ b/Engine/Plugins/Media/BlackmagicMedia/Source/BlackmagicMedia/Private/Player/BlackmagicMediaPlayer.cpp
@@ -450,18 +450,12 @@ FGuid FBlackmagicMediaPlayer::GetPlayerPluginGUID() const
  */
 bool FBlackmagicMediaPlayer::Open(const FString& Url, const IMediaOptions* Options)
 {
-	if (!FBlackmagic::IsInitialized())
+	if (!IBlackmagicMediaModule::Get().IsInitialized())
 	{
 		UE_LOG(LogBlackmagicMedia, Error, TEXT("The BlackmagicMediaPlayer can't open URL '%s'. Blackmagic is not initialized on your machine."), *Url);
 		return false;
 	}
 
-	if (!FBlackmagic::CanUseBlackmagicCard())
-	{
-		UE_LOG(LogBlackmagicMedia, Warning, TEXT("The BlackmagicMediaPlayer can't open URL '%s' because Blackmagic card cannot be used. Are you in a Commandlet? You may override this behavior by launching with -ForceBlackmagicUsage"), *Url);
-		return false;
-	}
-
 	if (!Super::Open(Url, Options))
 	{
 		return false;
