Fix audio glitches in blackmagic media output

diff --git a/Engine/Plugins/Media/AjaMedia/Source/AjaMediaOutput/Private/AjaMediaOutputModule.cpp b/Engine/Plugins/Media/AjaMedia/Source/AjaMediaOutput/Private/AjaMediaOutputModule.cpp
--- a/Engine/Plugins/Media/AjaMedia/Source/AjaMediaOutput/Private/AjaMediaOutputModule.cpp
+++ b/Engine/Plugins/Media/AjaMedia/Source/AjaMediaOutput/Private/AjaMediaOutputModule.cpp
@@ -19,7 +19,7 @@ void FAjaMediaOutputModule::StartupModule()
 	const auto DMAInitializationFunc = [this]()
 	{
 		const FGPUDriverInfo GPUDriverInfo = FPlatformMisc::GetGPUDriverInfo(GRHIAdapterName);
-		bIsGPUTextureTransferAvailable = GPUDriverInfo.IsNVIDIA();
+		bIsGPUTextureTransferAvailable = GPUDriverInfo.IsNVIDIA() &&!FModuleManager::Get().IsModuleLoaded("RenderDocPlugin");
 
 		if (bIsGPUTextureTransferAvailable)
 		{
diff --git a/Engine/Plugins/Media/BlackmagicMedia/Source/BlackmagicMediaOutput/Private/BlackmagicMediaCapture.cpp b/Engine/Plugins/Media/BlackmagicMedia/Source/BlackmagicMediaOutput/Private/BlackmagicMediaCapture.cpp
--- a/Engine/Plugins/Media/BlackmagicMedia/Source/BlackmagicMediaOutput/Private/BlackmagicMediaCapture.cpp
+++ b/Engine/Plugins/Media/BlackmagicMedia/Source/BlackmagicMediaOutput/Private/BlackmagicMediaCapture.cpp
@@ -541,7 +541,7 @@ bool UBlackmagicMediaCapture::InitBlackmagic(UBlackmagicMediaOutput* InBlackmagi
 		UMediaIOCoreSubsystem::FCreateAudioOutputArgs Args;
 		Args.NumOutputChannels = static_cast<int32>(ChannelOptions.NumAudioChannels);
 		Args.TargetFrameRate = FrameRate;
-		Args.MaxSampleLatency = InBlackmagicMediaOutput->AudioBufferSize;
+		Args.MaxSampleLatency = Align(InBlackmagicMediaOutput->AudioBufferSize, 4);
 		Args.OutputSampleRate = static_cast<uint32>(InBlackmagicMediaOutput->AudioSampleRate);
 		AudioOutput = GEngine->GetEngineSubsystem<UMediaIOCoreSubsystem>()->CreateAudioOutput(Args);
 	}
@@ -737,39 +737,51 @@ void UBlackmagicMediaCapture::WaitForSync_RenderingThread()
 
 void UBlackmagicMediaCapture::OutputAudio_RenderingThread(const FCaptureBaseData& InBaseData, const BlackmagicDesign::FTimecode& InTimecode)
 {
-	if (bOutputAudio)
+	if (bOutputAudio && AudioOutput)
 	{
 		TRACE_CPUPROFILER_EVENT_SCOPE(UBlackmagicMediaCapture::OutputAudio);
 
+		const double NewTimestamp = FPlatformTime::Seconds();
+		double CurrentFrameTime = NewTimestamp - OutputAudioTimestamp;
+		
+		const float TargetFrametime = 1 / FrameRate.AsDecimal();
+		if (UNLIKELY(OutputAudioTimestamp == 0))
+		{
+			CurrentFrameTime = TargetFrametime;
+		}
+
+		float FrameTimeRatio = CurrentFrameTime / TargetFrametime;
+
+		uint32_t NumSamplesToPull = FrameTimeRatio * AudioOutput->NumSamplesPerFrame;
+		NumSamplesToPull = FMath::Clamp(NumSamplesToPull, 0, AudioOutput->NumSamplesPerFrame);
+
 		BlackmagicDesign::FAudioSamplesDescriptor AudioSamples;
 		AudioSamples.Timecode = InTimecode;
 		AudioSamples.FrameIdentifier = InBaseData.SourceFrameNumber;
 		check(NumOutputChannels != 0);
-
-		if (AudioOutput)
+		
+		if (AudioBitDepth == EBlackmagicMediaOutputAudioBitDepth::Signed_32Bits)
 		{
-			if (AudioBitDepth == EBlackmagicMediaOutputAudioBitDepth::Signed_32Bits)
-			{
-				TArray<int32> AudioBuffer = AudioOutput->GetAudioSamples<int32>();
-				AudioSamples.AudioBuffer = reinterpret_cast<uint8_t*>(AudioBuffer.GetData());
-				AudioSamples.NumAudioSamples = AudioBuffer.Num() / NumOutputChannels;
-				AudioSamples.AudioBufferLength = AudioBuffer.Num() * sizeof(int32);
-				EventCallback->SendAudioSamples(AudioSamples);
-			}
-			else if (AudioBitDepth == EBlackmagicMediaOutputAudioBitDepth::Signed_16Bits)
-			{
-				TArray<int16> AudioBuffer = AudioOutput->GetAudioSamples<int16>();
-				AudioSamples.AudioBuffer = reinterpret_cast<uint8_t*>(AudioBuffer.GetData());
-				AudioSamples.NumAudioSamples = AudioBuffer.Num() / NumOutputChannels;
-				AudioSamples.AudioBufferLength = AudioBuffer.Num() * sizeof(int16);
+			TArray<int32> AudioBuffer = AudioOutput->GetAudioSamples<int32>(NumSamplesToPull);
+			AudioSamples.AudioBuffer = reinterpret_cast<uint8_t*>(AudioBuffer.GetData());
+			AudioSamples.NumAudioSamples = AudioBuffer.Num() / NumOutputChannels;
+			AudioSamples.AudioBufferLength = AudioBuffer.Num() * sizeof(int32);
+			EventCallback->SendAudioSamples(AudioSamples);
+		}
+		else if (AudioBitDepth == EBlackmagicMediaOutputAudioBitDepth::Signed_16Bits)
+		{
+			TArray<int16> AudioBuffer = AudioOutput->GetAudioSamples<int16>(NumSamplesToPull);
+			AudioSamples.AudioBuffer = reinterpret_cast<uint8_t*>(AudioBuffer.GetData());
+			AudioSamples.NumAudioSamples = AudioBuffer.Num() / NumOutputChannels;
+			AudioSamples.AudioBufferLength = AudioBuffer.Num() * sizeof(int16);
 
-				EventCallback->SendAudioSamples(AudioSamples);
-			}
-			else
-			{
-				checkNoEntry();
-			}
+			EventCallback->SendAudioSamples(AudioSamples);
 		}
+		else
+		{
+			checkNoEntry();
+		}
+		
+		OutputAudioTimestamp = NewTimestamp;
 	}
 }
-
diff --git a/Engine/Plugins/Media/BlackmagicMedia/Source/BlackmagicMediaOutput/Private/BlackmagicMediaOutputModule.cpp b/Engine/Plugins/Media/BlackmagicMedia/Source/BlackmagicMediaOutput/Private/BlackmagicMediaOutputModule.cpp
--- a/Engine/Plugins/Media/BlackmagicMedia/Source/BlackmagicMediaOutput/Private/BlackmagicMediaOutputModule.cpp
+++ b/Engine/Plugins/Media/BlackmagicMedia/Source/BlackmagicMediaOutput/Private/BlackmagicMediaOutputModule.cpp
@@ -19,7 +19,8 @@ void FBlackmagicMediaOutputModule::StartupModule()
 	const auto DMAInitializationFunc = [this]()
 	{
 		const FGPUDriverInfo GPUDriverInfo = FPlatformMisc::GetGPUDriverInfo(GRHIAdapterName);
-		bIsGPUTextureTransferAvailable = GPUDriverInfo.IsNVIDIA();
+		bIsGPUTextureTransferAvailable = GPUDriverInfo.IsNVIDIA() && !FModuleManager::Get().IsModuleLoaded("RenderDocPlugin");
+
 
 		if (bIsGPUTextureTransferAvailable)
 		{
diff --git a/Engine/Plugins/Media/BlackmagicMedia/Source/BlackmagicMediaOutput/Public/BlackmagicMediaCapture.h b/Engine/Plugins/Media/BlackmagicMedia/Source/BlackmagicMediaOutput/Public/BlackmagicMediaCapture.h
--- a/Engine/Plugins/Media/BlackmagicMedia/Source/BlackmagicMediaOutput/Public/BlackmagicMediaCapture.h
+++ b/Engine/Plugins/Media/BlackmagicMedia/Source/BlackmagicMediaOutput/Public/BlackmagicMediaCapture.h
@@ -91,6 +91,8 @@ private:
 
 	EBlackmagicMediaOutputAudioBitDepth AudioBitDepth = EBlackmagicMediaOutputAudioBitDepth::Signed_16Bits;
 
+	uint32 NumAudioSamplesPerFrame = 0;
+
 	bool bOutputAudio = false;
 
 	uint8 NumOutputChannels = 0;
@@ -100,4 +102,7 @@ private:
 
 	/** Whether or not GPUTextureTransfer was initialized successfully. */
 	bool bGPUTextureTransferAvailable = false;
+
+	/** The last time OutputAudio was called. Used to adjust the number of audio samples to grab on each frame depending on the frame time. */
+	double OutputAudioTimestamp;
 };
diff --git a/Engine/Plugins/Media/MediaIOFramework/Source/MediaIOCore/Private/MediaIOCoreAudioOutput.cpp b/Engine/Plugins/Media/MediaIOFramework/Source/MediaIOCore/Private/MediaIOCoreAudioOutput.cpp
--- a/Engine/Plugins/Media/MediaIOFramework/Source/MediaIOCore/Private/MediaIOCoreAudioOutput.cpp
+++ b/Engine/Plugins/Media/MediaIOFramework/Source/MediaIOCore/Private/MediaIOCoreAudioOutput.cpp
@@ -22,6 +22,7 @@ FMediaIOAudioOutput::FMediaIOAudioOutput(Audio::FPatchOutputStrongPtr InPatchOut
     , MaxSampleLatency(InAudioOptions.InMaxSampleLatency)
     , OutputSampleRate(InAudioOptions.InOutputSampleRate)
 {
+	NumSamplesPerFrame = FMath::CeilToInt(NumInputChannels * OutputSampleRate / TargetFrameRate.AsDecimal());
 }
 
 int32 FMediaIOAudioOutput::GetAudioBuffer(int32 InNumSamplesToPop, float* OutBuffer) const
@@ -36,22 +37,16 @@ int32 FMediaIOAudioOutput::GetAudioBuffer(int32 InNumSamplesToPop, float* OutBuf
 	return 0;
 }
 
-Audio::FAlignedFloatBuffer FMediaIOAudioOutput::GetFloatBuffer() const
+Audio::FAlignedFloatBuffer FMediaIOAudioOutput::GetFloatBuffer(uint32 NumSamplesToGet) const
 {
-	// @todo: Depend on frame number to correctly fetch the right amount of frames on framerates like 59.97
-	const int32 NumSamplesPerFrame = FMath::CeilToInt(NumInputChannels * OutputSampleRate / TargetFrameRate.AsDecimal());
-
 	// NumSamplesToPop must be a multiple of 4 in order to avoid an assertion in the audio pipeline.
-	const int32 NumSamplesToPop = Align(NumSamplesPerFrame, 4);
+	const int32 NumSamplesToPop = Align(NumSamplesToGet, 4);
 
 	Audio::FAlignedFloatBuffer FloatBuffer;
 	FloatBuffer.SetNumZeroed(NumSamplesToPop);
 
 	const int32 NumPopped = GetAudioBuffer(NumSamplesToPop, FloatBuffer.GetData());
 
-	// Trim back the buffer after we get the buffer since it might be bigger because of alignment
-	constexpr bool bAllowShrinking = false;
-	FloatBuffer.SetNum(NumPopped, bAllowShrinking);
 	return FloatBuffer;
 }
 
diff --git a/Engine/Plugins/Media/MediaIOFramework/Source/MediaIOCore/Public/MediaIOCoreAudioOutput.h b/Engine/Plugins/Media/MediaIOFramework/Source/MediaIOCore/Public/MediaIOCoreAudioOutput.h
--- a/Engine/Plugins/Media/MediaIOFramework/Source/MediaIOCore/Public/MediaIOCoreAudioOutput.h
+++ b/Engine/Plugins/Media/MediaIOFramework/Source/MediaIOCore/Public/MediaIOCoreAudioOutput.h
@@ -9,7 +9,7 @@
 namespace UE::MediaIoCoreModule::Private
 {
 	template <typename OutputType>
-	TArray<OutputType> ConvertAndUpmixBuffer(const Audio::FAlignedFloatBuffer& InBuffer, int32 NumInputChannels, int32 NumOutputChannels)
+	TArray<OutputType> ConvertAndUpmixBuffer(const Audio::FAlignedFloatBuffer& InBuffer, int32 NumInputChannels, int32 NumOutputChannels, uint32 NumSamplesToGet)
 	{
 		TRACE_CPUPROFILER_EVENT_SCOPE(MediaIoCoreModule::ConvertAndUpmixBuffer);
 		/**
@@ -42,6 +42,9 @@ namespace UE::MediaIoCoreModule::Private
 		
 			ConvertedBufferPtr += NumOutputChannels;
 		}
+
+		constexpr bool bAllowShrinking = false;
+		ConvertedBuffer.SetNum(NumSamplesToGet, bAllowShrinking);
 		return ConvertedBuffer;
 	}
 }
@@ -68,13 +71,26 @@ public:
 	{
 		TRACE_CPUPROFILER_EVENT_SCOPE(FMediaIOAudioOutput::GetAudioSamples);
 		
-		const Audio::FAlignedFloatBuffer FloatBuffer = GetFloatBuffer();
-		return UE::MediaIoCoreModule::Private::ConvertAndUpmixBuffer<OutputType>(FloatBuffer, NumInputChannels, NumOutputChannels);
+		const Audio::FAlignedFloatBuffer FloatBuffer = GetFloatBuffer(NumSamplesPerFrame);
+		return UE::MediaIoCoreModule::Private::ConvertAndUpmixBuffer<OutputType>(FloatBuffer, NumInputChannels, NumOutputChannels, NumSamplesPerFrame);
+	}
+	
+	template <typename OutputType>
+	TArray<OutputType> GetAudioSamples(uint32 NumSamplesToGet) const
+	{
+		TRACE_CPUPROFILER_EVENT_SCOPE(FMediaIOAudioOutput::GetAudioSamples);
+		
+		const Audio::FAlignedFloatBuffer FloatBuffer = GetFloatBuffer(NumSamplesToGet);
+		return UE::MediaIoCoreModule::Private::ConvertAndUpmixBuffer<OutputType>(FloatBuffer, NumInputChannels, NumOutputChannels, NumSamplesToGet);
 	}
 
+public:
+	// @todo: Depend on frame number to correctly fetch the right amount of frames on framerates like 59.97
+	int32 NumSamplesPerFrame = 0;
+
 private:
 	int32 GetAudioBuffer(int32 InNumSamplesToPop, float* OutBuffer) const;
-	Audio::FAlignedFloatBuffer GetFloatBuffer() const;
+	Audio::FAlignedFloatBuffer GetFloatBuffer(uint32 NumSamplesToGet) const;
 
 private:
 	/** The buffer accumulating audio samples. */
