#u5 - Fix build, CL 19098250 came in just before 19098314 so there was a conflict with FMath::Clamp() needing an explicit type.

diff --git a/Engine/Plugins/Media/BlackmagicMedia/Source/BlackmagicMediaOutput/Private/BlackmagicMediaCapture.cpp b/Engine/Plugins/Media/BlackmagicMedia/Source/BlackmagicMediaOutput/Private/BlackmagicMediaCapture.cpp
--- a/Engine/Plugins/Media/BlackmagicMedia/Source/BlackmagicMediaOutput/Private/BlackmagicMediaCapture.cpp
+++ b/Engine/Plugins/Media/BlackmagicMedia/Source/BlackmagicMediaOutput/Private/BlackmagicMediaCapture.cpp
@@ -753,7 +753,7 @@ void UBlackmagicMediaCapture::OutputAudio_RenderingThread(const FCaptureBaseData
 		float FrameTimeRatio = CurrentFrameTime / TargetFrametime;
 
 		uint32_t NumSamplesToPull = FrameTimeRatio * AudioOutput->NumSamplesPerFrame;
-		NumSamplesToPull = FMath::Clamp(NumSamplesToPull, 0, AudioOutput->NumSamplesPerFrame);
+		NumSamplesToPull = FMath::Clamp<uint32_t>(NumSamplesToPull, 0, AudioOutput->NumSamplesPerFrame);
 
 		BlackmagicDesign::FAudioSamplesDescriptor AudioSamples;
 		AudioSamples.Timecode = InTimecode;
