Disable GPU Texture Transfer on Nvidia tesla GPUs for the moment to fix build

#rnx

diff --git a/Engine/Plugins/Media/AjaMedia/Source/AjaMediaOutput/Private/AjaMediaOutputModule.cpp b/Engine/Plugins/Media/AjaMedia/Source/AjaMediaOutput/Private/AjaMediaOutputModule.cpp
--- a/Engine/Plugins/Media/AjaMedia/Source/AjaMediaOutput/Private/AjaMediaOutputModule.cpp
+++ b/Engine/Plugins/Media/AjaMedia/Source/AjaMediaOutput/Private/AjaMediaOutputModule.cpp
@@ -20,6 +20,7 @@ void FAjaMediaOutputModule::StartupModule()
 	{
 		const FGPUDriverInfo GPUDriverInfo = FPlatformMisc::GetGPUDriverInfo(GRHIAdapterName);
 		bIsGPUTextureTransferAvailable = GPUDriverInfo.IsNVIDIA() &&!FModuleManager::Get().IsModuleLoaded("RenderDocPlugin");
+		bIsGPUTextureTransferAvailable &= !GPUDriverInfo.DeviceDescription.Contains(TEXT("Tesla"));
 
 		if (bIsGPUTextureTransferAvailable)
 		{
diff --git a/Engine/Plugins/Media/BlackmagicMedia/Source/BlackmagicMediaOutput/Private/BlackmagicMediaOutputModule.cpp b/Engine/Plugins/Media/BlackmagicMedia/Source/BlackmagicMediaOutput/Private/BlackmagicMediaOutputModule.cpp
--- a/Engine/Plugins/Media/BlackmagicMedia/Source/BlackmagicMediaOutput/Private/BlackmagicMediaOutputModule.cpp
+++ b/Engine/Plugins/Media/BlackmagicMedia/Source/BlackmagicMediaOutput/Private/BlackmagicMediaOutputModule.cpp
@@ -20,7 +20,7 @@ void FBlackmagicMediaOutputModule::StartupModule()
 	{
 		const FGPUDriverInfo GPUDriverInfo = FPlatformMisc::GetGPUDriverInfo(GRHIAdapterName);
 		bIsGPUTextureTransferAvailable = GPUDriverInfo.IsNVIDIA() && !FModuleManager::Get().IsModuleLoaded("RenderDocPlugin");
-
+		bIsGPUTextureTransferAvailable &= !GPUDriverInfo.DeviceDescription.Contains(TEXT("Tesla"));
 
 		if (bIsGPUTextureTransferAvailable)
 		{
