Put common logic in FMediaIOCorePlayerBase

diff --git a/Engine/Plugins/Media/AjaMedia/Source/AjaMedia/AjaMedia.Build.cs b/Engine/Plugins/Media/AjaMedia/Source/AjaMedia/AjaMedia.Build.cs
--- a/Engine/Plugins/Media/AjaMedia/Source/AjaMedia/AjaMedia.Build.cs
+++ b/Engine/Plugins/Media/AjaMedia/Source/AjaMedia/AjaMedia.Build.cs
@@ -23,7 +23,6 @@ namespace UnrealBuildTool.Rules
 					"Core",
 					"CoreUObject",
 					"Engine",
-					"GPUTextureTransfer",
 					"MediaIOCore",
 					"MediaUtils",
 					"Projects",
diff --git a/Engine/Plugins/Media/AjaMedia/Source/AjaMedia/Private/Player/AjaMediaPlayer.cpp b/Engine/Plugins/Media/AjaMedia/Source/AjaMedia/Private/Player/AjaMediaPlayer.cpp
--- a/Engine/Plugins/Media/AjaMedia/Source/AjaMedia/Private/Player/AjaMediaPlayer.cpp
+++ b/Engine/Plugins/Media/AjaMedia/Source/AjaMedia/Private/Player/AjaMediaPlayer.cpp
@@ -7,7 +7,6 @@
 #include "MediaIOCoreEncodeTime.h"
 #include "MediaIOCoreFileWriter.h"
 #include "MediaIOCoreSamples.h"
-#include "GPUTextureTransferModule.h"
 #include "HAL/PlatformAtomics.h"
 #include "HAL/PlatformProcess.h"
 #include "IAjaMediaModule.h"
@@ -46,11 +45,6 @@ static FAutoConsoleCommand AjaWriteOutputRawDataCmd(
 	FConsoleCommandDelegate::CreateLambda([]() { bAjaWriteOutputRawDataCmdEnable = true; })
 	);
 
-static TAutoConsoleVariable<int32> CVarAjaEnableGPUDirect(
-	TEXT("Aja.EnableGPUDirectInput"), 0,
-	TEXT("Whether to enable GPU direct for faster video frame copies. (Experimental)"),
-	ECVF_RenderThreadSafe);
-
 /* FAjaVideoPlayer structors
  *****************************************************************************/
 
@@ -80,8 +74,6 @@ FAjaMediaPlayer::FAjaMediaPlayer(IMediaEventSink& InEventSink)
 	, SupportedSampleTypes(EMediaIOSampleType::None)
 	, bPauseRequested(false)
 {
-	FGPUTextureTransferModule& Module = FGPUTextureTransferModule::Get();
-	GPUTextureTransfer = Module.GetTextureTransfer();
 }
 
 
@@ -167,8 +159,8 @@ bool FAjaMediaPlayer::Open(const FString& Url, const IMediaOptions* Options)
 		LastVideoFormatIndex = AjaOptions.VideoFormatIndex;
 	}
 	{
-		const EAjaMediaSourceColorFormat ColorFormat = (EAjaMediaSourceColorFormat)(Options->GetMediaOption(AjaMediaOption::ColorFormat, (int64)EAjaMediaSourceColorFormat::YUV2_8bit));
-		switch(ColorFormat)
+		AjaColorFormat= (EAjaMediaSourceColorFormat)(Options->GetMediaOption(AjaMediaOption::ColorFormat, (int64)EAjaMediaSourceColorFormat::YUV2_8bit));
+		switch(AjaColorFormat)
 		{
 		case EAjaMediaSourceColorFormat::YUV2_8bit:
 			if (AjaOptions.bUseKey)
@@ -217,7 +209,6 @@ bool FAjaMediaPlayer::Open(const FString& Url, const IMediaOptions* Options)
 	MaxNumAudioFrameBuffer = Options->GetMediaOption(AjaMediaOption::MaxAudioFrameBuffer, (int64)8);
 	MaxNumMetadataFrameBuffer = Options->GetMediaOption(AjaMediaOption::MaxAncillaryFrameBuffer, (int64)8);
 	MaxNumVideoFrameBuffer = Options->GetMediaOption(AjaMediaOption::MaxVideoFrameBuffer, (int64)8);
-	PixelFormat = (EAjaMediaSourceColorFormat)(Options->GetMediaOption(AjaMediaOption::ColorFormat, (int64)EAjaMediaSourceColorFormat::YUV2_8bit));
 
 	check(InputChannel == nullptr);
 	InputChannel = new AJA::AJAInputChannel();
@@ -260,12 +251,6 @@ bool FAjaMediaPlayer::Open(const FString& Url, const IMediaOptions* Options)
 	}
 #endif
 
-
-	if (CVarAjaEnableGPUDirect.GetValueOnAnyThread())
-	{
-		CreateAndRegisterTextures(Options);
-	}
-
 	return true;
 }
 
@@ -526,7 +511,9 @@ bool FAjaMediaPlayer::OnRequestInputBuffer(const AJA::AJARequestInputBufferData&
 	// Video
 	if (bUseVideo && InRequestBuffer.VideoBufferSize > 0 && InRequestBuffer.bIsProgressivePicture)
 	{
-		if (CVarAjaEnableGPUDirect.GetValueOnAnyThread())
+		const bool bCanUseGPUTextureTransfer = CanUseGPUTextureTransfer();
+
+		if (bCanUseGPUTextureTransfer)
 		{
 			if (!Textures.Num())
 			{
@@ -538,28 +525,9 @@ bool FAjaMediaPlayer::OnRequestInputBuffer(const AJA::AJARequestInputBufferData&
 		AjaThreadCurrentTextureSample = TextureSamplePool->AcquireShared();
 		AjaThreadCurrentTextureSample->RequestBuffer(InRequestBuffer.VideoBufferSize);
 
-		if (CVarAjaEnableGPUDirect.GetValueOnAnyThread())
+		if (bCanUseGPUTextureTransfer)
 		{
-			FScopeLock Lock(&TexturesCriticalSection);
-			TRefCountPtr<FRHITexture> Texture = Textures.Pop();
-			if (Texture && RegisteredTextures.Contains(Texture))
-			{
-				AjaThreadCurrentTextureSample->SetTexture(Texture);
-
-				AjaThreadCurrentTextureSample->SetDestructionCallback([MediaPlayerWeakPtr = TWeakPtr<FAjaMediaPlayer>(AsShared())](TRefCountPtr<FRHITexture> InTexture)
-				{
-					if (TSharedPtr<FAjaMediaPlayer> MediaPlayer = MediaPlayerWeakPtr.Pin())
-					{
-						MediaPlayer->OnSampleDestroyed(InTexture);
-					}
-				});
-			}
-			else
-			{
-				UE_LOG(LogAjaMedia, Error, TEXT("Unregistered texture %u encountered while doing a gpu texture transfer."), Texture ? reinterpret_cast<std::uintptr_t>(Texture->GetNativeResource()) : 0);
-			}
-
-			RegisterSampleBuffer(AjaThreadCurrentTextureSample);
+			PreGPUTransfer(AjaThreadCurrentTextureSample);
 		}
 
 		OutRequestedBuffer.VideoBuffer = reinterpret_cast<uint8_t*>(AjaThreadCurrentTextureSample->GetMutableBuffer());
@@ -715,7 +683,7 @@ bool FAjaMediaPlayer::OnInputFrameReceived(const AJA::AJAInputFrameData& InInput
 			break;
 		}
 
-		if (bEncodeTimecodeInTexel && DecodedTimecode.IsSet() && InVideoFrame.bIsProgressivePicture && !CVarAjaEnableGPUDirect.GetValueOnAnyThread())
+		if (bEncodeTimecodeInTexel && DecodedTimecode.IsSet() && InVideoFrame.bIsProgressivePicture)
 		{
 			FTimecode SetTimecode = DecodedTimecode.GetValue();
 			FMediaIOCoreEncodeTime EncodeTime(EncodePixelFormat, InVideoFrame.VideoBuffer, InVideoFrame.Stride, InVideoFrame.Width, InVideoFrame.Height);
@@ -732,25 +700,9 @@ bool FAjaMediaPlayer::OnInputFrameReceived(const AJA::AJAInputFrameData& InInput
 		{
 			if (AjaThreadCurrentTextureSample->SetProperties(InVideoFrame.Stride, InVideoFrame.Width, InVideoFrame.Height, VideoSampleFormat, DecodedTime, VideoFrameRate, DecodedTimecode, bIsSRGBInput))
 			{
-				if (GPUTextureTransfer && CVarAjaEnableGPUDirect.GetValueOnAnyThread())
+				if (CanUseGPUTextureTransfer())
 				{
-					// We must start transfer on render thread
-					ENQUEUE_RENDER_COMMAND(StartGPUTextureTransferCopy)(
-						[this, CurrentSample = AjaThreadCurrentTextureSample](FRHICommandListImmediate& RHICmdList)
-						{
-							if (CurrentSample->GetTexture())
-							{
-								if (GPUTextureTransfer)
-								{
-									void* Buffer = CurrentSample->GetMutableBuffer();
-									UE_LOG(LogAjaMedia, Verbose, TEXT("Starting Transfer with buffer %u"), reinterpret_cast<std::uintptr_t>(Buffer));
-									GPUTextureTransfer->TransferTexture(Buffer, CurrentSample->GetTexture(), UE::GPUTextureTransfer::ETransferDirection::CPU_TO_GPU);
-									GPUTextureTransfer->BeginSync((CurrentSample)->GetMutableBuffer(), UE::GPUTextureTransfer::ETransferDirection::CPU_TO_GPU);
-									Samples->AddVideo(CurrentSample.ToSharedRef());
-									GPUTextureTransfer->EndSync(CurrentSample->GetMutableBuffer());
-								}
-							}
-						});
+					ExecuteGPUTransfer(AjaThreadCurrentTextureSample);
 				}
 				else
 				{
@@ -822,174 +774,9 @@ void FAjaMediaPlayer::SetupSampleChannels()
 	Samples->InitializeMetadataBuffer(MetadataSettings);
 }
 
-void FAjaMediaPlayer::OnSampleDestroyed(TRefCountPtr<FRHITexture> InTexture)
-{
-	FScopeLock ScopeLock(&TexturesCriticalSection);
-	if (InTexture && RegisteredTextures.Contains(InTexture))
-	{
-		Textures.Add(InTexture);
-	}
-}
-
-void FAjaMediaPlayer::RegisterSampleBuffer(const TSharedPtr<FAjaMediaTextureSample>& InSample)
-{
-	if (!RegisteredBuffers.Contains(InSample->GetMutableBuffer()))
-	{
-		RegisteredBuffers.Add(InSample->GetMutableBuffer());
-
-		// Enqueue buffer registration on the render  thread to ensure it happens before the call to TransferTexture
-		ENQUEUE_RENDER_COMMAND(RegisterSampleBuffers)([this, InSample](FRHICommandList& CommandList)
-			{
-				if (!InSample->GetMutableBuffer())
-				{
-					UE_LOG(LogAjaMedia, Error, TEXT("A buffer was not available while performing a gpu texture transfer."));
-					return;
-				}
-
-				if (!InSample->GetTexture())
-				{
-					UE_LOG(LogAjaMedia, Error, TEXT("A texture was not available while performing a gpu texture transfer."));
-					return;
-				}
-
-				const uint32 TextureWidth = InSample->GetTexture()->GetDesc().Extent.X;
-				const uint32 TextureHeight = InSample->GetTexture()->GetDesc().Extent.Y;
-				uint32 TextureStride = TextureWidth * 4;
-
-				EPixelFormat Format = InSample->GetTexture()->GetFormat();
-				if (Format == PF_R32G32B32A32_UINT)
-				{
-					TextureStride *= 4;
-				}
-
-				UE_LOG(LogAjaMedia, Verbose, TEXT("Registering buffer %u"), reinterpret_cast<std::uintptr_t>(InSample->GetMutableBuffer()));
-				UE::GPUTextureTransfer::FRegisterDMABufferArgs Args;
-				Args.Buffer = InSample->GetMutableBuffer();
-				Args.Width = TextureWidth;
-				Args.Height = TextureHeight;
-				Args.Stride = TextureStride;
-				Args.PixelFormat = Format == PF_B8G8R8A8 ? UE::GPUTextureTransfer::EPixelFormat::PF_8Bit : UE::GPUTextureTransfer::EPixelFormat::PF_10Bit;
-				GPUTextureTransfer->RegisterBuffer(Args);
-			});
-	}
-	else
-	{
-		UE_LOG(LogAjaMedia, VeryVerbose, TEXT("Buffer already registered: %u"), reinterpret_cast<std::uintptr_t>(AjaThreadCurrentTextureSample->GetMutableBuffer()));
-	}
-}
-
-void FAjaMediaPlayer::UnregisterSampleBuffers()
-{
-	ENQUEUE_RENDER_COMMAND(UnregisterSampleBuffers)([TextureTransfer =GPUTextureTransfer, BuffersToUnregister = RegisteredBuffers](FRHICommandList& CommandList)
-	{
-		if (TextureTransfer)
-		{
-			for (void* RegisteredBuffer : BuffersToUnregister)
-			{
-				TextureTransfer->UnregisterBuffer(RegisteredBuffer);
-			}
-		}
-	});
-	
-	RegisteredBuffers.Reset();
-}
-
-void FAjaMediaPlayer::CreateAndRegisterTextures(const IMediaOptions* Options)
-{
-	ENQUEUE_RENDER_COMMAND(AjaMediaPlayerCreateTextures)(
-		[this, Options](FRHICommandListImmediate& CommandList)
-		{
-			if (!GPUTextureTransfer)
-			{
-				UE_LOG(LogAjaMedia, Error, TEXT("Could not register textures for gpu texture transfer, as the GPUTextureTransfer object was null."));
-				return;
-			}
-
-			for (uint8 Index = 0; Index < MaxNumVideoFrameBuffer; Index++)
-			{
-				TRefCountPtr<FRHITexture> RHITexture;
-				const FString TextureName = FString::Printf(TEXT("FAjaMediaPlayerTexture %d"), Index);
-				FRHIResourceCreateInfo CreateInfo(*TextureName);
-				
-				uint32 TextureWidth = VideoTrackFormat.Dim.X;
-				const uint32 TextureHeight = VideoTrackFormat.Dim.Y;
-				uint32 Stride = TextureWidth;
-				EPixelFormat InputFormat = EPixelFormat::PF_Unknown;
-
-				if (PixelFormat == EAjaMediaSourceColorFormat::YUV2_8bit)
-				{
-					InputFormat = EPixelFormat::PF_B8G8R8A8;
-					TextureWidth /= 2;
-					Stride = TextureWidth * 4;
-				}
-				else if (PixelFormat == EAjaMediaSourceColorFormat::YUV_10bit)
-				{
-					InputFormat = EPixelFormat::PF_R32G32B32A32_UINT;
-					TextureWidth /= 6;
-					Stride = TextureWidth * 16;
-				}
-				else
-				{
-					checkf(false, TEXT("Format not supported"));
-				}
-
-				TRefCountPtr<FRHITexture2D> DummyTexture2DRHI;
-
-				ETextureCreateFlags CreateFlags = TexCreate_Shared;
-				if (RHIGetInterfaceType() == ERHIInterfaceType::Vulkan)
-				{
-					CreateFlags = TexCreate_External;
-				}
-
-				FRHITextureCreateDesc CreateDesc = FRHITextureCreateDesc::Create2D(*TextureName, FIntPoint(TextureWidth, TextureHeight), InputFormat, FClearValueBinding::White, CreateFlags);
-				RHICreateTargetableShaderResource(CreateDesc, TexCreate_RenderTargetable, RHITexture, DummyTexture2DRHI);
-
-				UE_LOG(LogAjaMedia, Verbose, TEXT("Registering texture %u"), reinterpret_cast<std::uintptr_t>(RHITexture->GetNativeResource()));
-
-				UE::GPUTextureTransfer::FRegisterDMATextureArgs Args;
-				Args.RHITexture = RHITexture.GetReference();
-				Args.Stride = Stride;
-				Args.Width = TextureWidth;
-				Args.Height = TextureHeight;
-				Args.RHIResourceMemory = nullptr;
-				Args.PixelFormat = PixelFormat == EAjaMediaSourceColorFormat::YUV2_8bit ? UE::GPUTextureTransfer::EPixelFormat::PF_8Bit : UE::GPUTextureTransfer::EPixelFormat::PF_10Bit;
-				
-				GPUTextureTransfer->RegisterTexture(Args);
-				RegisteredTextures.Add(RHITexture);
-
-				{
-					FScopeLock Lock(&TexturesCriticalSection);
-					Textures.Add(RHITexture);
-				}
-			}
-		});
-
-	FRenderCommandFence RenderFence;
-	RenderFence.BeginFence();
-	RenderFence.Wait();
-}
-
-void FAjaMediaPlayer::UnregisterTextures()
+void FAjaMediaPlayer::AddVideoSample(const TSharedRef<FMediaIOCoreTextureSampleBase>& InSample)
 {
-	// Keep a copy of the textures to make sure they were not destoryed before this gets executed.
-	ENQUEUE_RENDER_COMMAND(UnregisterSampleBuffers)([TextureTransfer = GPUTextureTransfer, TexturesToUnregister = RegisteredTextures](FRHICommandList& CommandList)
-		{
-			if (TextureTransfer)
-			{
-				for (const TRefCountPtr<FRHITexture>& RHITexture : TexturesToUnregister)
-				{
-					UE_LOG(LogAjaMedia, Verbose, TEXT("Unregistering texture %u"), reinterpret_cast<std::uintptr_t>(RHITexture->GetNativeResource()));
-					TextureTransfer->UnregisterTexture(RHITexture.GetReference());
-				}
-			}
-		});
-
-	RegisteredTextures.Reset();
-
-	{
-		FScopeLock Lock(&TexturesCriticalSection);
-		Textures.Reset();
-	}
+	Samples->AddVideo(InSample);
 }
 
 bool FAjaMediaPlayer::SetRate(float Rate)
diff --git a/Engine/Plugins/Media/AjaMedia/Source/AjaMedia/Private/Player/AjaMediaPlayer.h b/Engine/Plugins/Media/AjaMedia/Source/AjaMedia/Private/Player/AjaMediaPlayer.h
--- a/Engine/Plugins/Media/AjaMedia/Source/AjaMedia/Private/Player/AjaMediaPlayer.h
+++ b/Engine/Plugins/Media/AjaMedia/Source/AjaMedia/Private/Player/AjaMediaPlayer.h
@@ -6,8 +6,6 @@
 
 #include "AjaMediaPrivate.h"
 #include "AjaMediaSource.h"
-#include "GPUTextureTransferModule.h"
-#include "GPUTextureTransfer.h"
 #include "HAL/CriticalSection.h"
 #include "Templates/PimplPtr.h"
 
@@ -43,7 +41,6 @@ namespace AJA
  */
 class FAjaMediaPlayer
 	: public FMediaIOCorePlayerBase
-	, public TSharedFromThis<FAjaMediaPlayer>
 	, protected AJA::IAJAInputOutputChannelCallbackInterface
 {
 	using Super = FMediaIOCorePlayerBase;
@@ -102,18 +99,23 @@ protected:
 	/** Verify if we lost some frames since last Tick*/
 	void VerifyFrameDropCount();
 
-
 	virtual bool IsHardwareReady() const override;
 
 	//~ FMediaIOCorePlayerBase interface
 	virtual void SetupSampleChannels() override;
 
-private:
-	void OnSampleDestroyed(TRefCountPtr<FRHITexture> InTexture);
-	void RegisterSampleBuffer(const TSharedPtr<FAjaMediaTextureSample>& InSample);
-	void UnregisterSampleBuffers();
-	void CreateAndRegisterTextures(const IMediaOptions* Options);
-	void UnregisterTextures();
+	virtual uint32 GetNumVideoFrameBuffers() const override
+	{
+		return MaxNumVideoFrameBuffer;
+	}
+
+	virtual EMediaIOCoreColorFormat GetColorFormat() const override
+	{
+		return AjaColorFormat == EAjaMediaSourceColorFormat::YUV2_8bit ? EMediaIOCoreColorFormat::YUV8 : EMediaIOCoreColorFormat::YUV10;
+	}
+
+	virtual void AddVideoSample(const TSharedRef<FMediaIOCoreTextureSampleBase>& InSample) override;
+
 private:
 
 	/** Audio, MetaData, Texture  sample object pool. */
@@ -183,17 +185,5 @@ private:
 	/** Flag to indicate that pause is being requested */
 	std::atomic<bool> bPauseRequested;
 
-	FCriticalSection TexturesCriticalSection;
-	
-	/** GPU Texture transfer object */
-	UE::GPUTextureTransfer::TextureTransferPtr GPUTextureTransfer;
-
-	/** Pool of textures registerd with GPU Texture transfer. */
-	TArray<TRefCountPtr<FRHITexture>> Textures;
-	/** Buffers Registered with GPU Texture Transfer */
-	TSet<void*> RegisteredBuffers;
-	/** Pool of textures registerd with GPU Texture transfer. */
-	TSet<TRefCountPtr<FRHITexture>> RegisteredTextures;
-
-	EAjaMediaSourceColorFormat PixelFormat = EAjaMediaSourceColorFormat::YUV2_8bit;
+	EAjaMediaSourceColorFormat AjaColorFormat = EAjaMediaSourceColorFormat::YUV2_8bit;
 };
diff --git a/Engine/Plugins/Media/AjaMedia/Source/AjaMedia/Private/Player/AjaMediaTextureSample.h b/Engine/Plugins/Media/AjaMedia/Source/AjaMedia/Private/Player/AjaMediaTextureSample.h
--- a/Engine/Plugins/Media/AjaMedia/Source/AjaMedia/Private/Player/AjaMediaTextureSample.h
+++ b/Engine/Plugins/Media/AjaMedia/Source/AjaMedia/Private/Player/AjaMediaTextureSample.h
@@ -14,17 +14,6 @@ class FAjaMediaTextureSample
 	using Super = FMediaIOCoreTextureSampleBase;
 
 public:
-
-	virtual void ShutdownPoolable() override
-	{
-		if (DestructionCallback)
-		{
-			DestructionCallback(Texture);
-		}
-
-		Super::FreeSample();
-	}
-
 	/**
 	 * Initialize the sample.
 	 *
@@ -75,16 +64,6 @@ public:
 			, bInIsSRGB);
 	}
 
-	void SetTexture(TRefCountPtr<FRHITexture> InRHITexture)
-	{
-		Texture = MoveTemp(InRHITexture);
-	}
-
-	void SetDestructionCallback(TFunction<void(TRefCountPtr<FRHITexture2D>)> InDestructionCallback)
-	{
-		DestructionCallback = InDestructionCallback;
-	}
-
 	/**
 	 * Get YUV to RGB conversion matrix
 	 *
@@ -94,40 +73,6 @@ public:
 	{
 		return MediaShaders::YuvToRgbRec709Scaled;
 	}
-
-	virtual const void* GetBuffer() override
-	{
-		// Don't return the buffer if we have a texture to force the media player to use the texture if available. 
-		if (Texture)
-		{
-			return nullptr;
-		}
-		return Buffer.GetData();
-
-	}
-
-	void* GetMutableBuffer()
-	{
-		return Buffer.GetData();
-	}
-
-#if WITH_ENGINE
-	virtual FRHITexture* GetTexture() const override
-	{
-		if (Texture)
-		{
-			return Texture.GetReference();
-		}
-		return nullptr;
-	}
-#endif //WITH_ENGINE
-
-private:
-	/** Hold a texture to be used for gpu texture transfers. */
-	TRefCountPtr<FRHITexture2D> Texture;
-
-	/** Called when the sample is destroyed by its pool. */
-	TFunction<void(TRefCountPtr<FRHITexture2D>)> DestructionCallback;
 };
 
 /*
diff --git a/Engine/Plugins/Media/BlackmagicMedia/Source/BlackmagicMedia/BlackmagicMedia.Build.cs b/Engine/Plugins/Media/BlackmagicMedia/Source/BlackmagicMedia/BlackmagicMedia.Build.cs
--- a/Engine/Plugins/Media/BlackmagicMedia/Source/BlackmagicMedia/BlackmagicMedia.Build.cs
+++ b/Engine/Plugins/Media/BlackmagicMedia/Source/BlackmagicMedia/BlackmagicMedia.Build.cs
@@ -32,7 +32,6 @@ namespace UnrealBuildTool.Rules
 				new string[] {
 					"Blackmagic",
 					"Engine",
-					"GPUTextureTransfer",
 					"MediaUtils",
 					"Projects",
 					"SlateCore"
diff --git a/Engine/Plugins/Media/BlackmagicMedia/Source/BlackmagicMedia/Private/Player/BlackmagicMediaPlayer.cpp b/Engine/Plugins/Media/BlackmagicMedia/Source/BlackmagicMedia/Private/Player/BlackmagicMediaPlayer.cpp
--- a/Engine/Plugins/Media/BlackmagicMedia/Source/BlackmagicMedia/Private/Player/BlackmagicMediaPlayer.cpp
+++ b/Engine/Plugins/Media/BlackmagicMedia/Source/BlackmagicMedia/Private/Player/BlackmagicMediaPlayer.cpp
@@ -7,7 +7,6 @@
 #include "BlackmagicMediaSource.h"
 
 #include "Engine/GameEngine.h"
-#include "GPUTextureTransferModule.h"
 #include "HAL/CriticalSection.h"
 #include "HAL/PlatformProcess.h"
 #include "IBlackmagicMediaModule.h"
@@ -35,13 +34,6 @@
 
 DECLARE_CYCLE_STAT(TEXT("Blackmagic MediaPlayer Process received frame"), STAT_Blackmagic_MediaPlayer_ProcessReceivedFrame, STATGROUP_Media);
 
-
-static TAutoConsoleVariable<int32> CVarBlackmagicEnableGPUDirect(
-	TEXT("Blackmagic.EnableGPUDirectInput"), 0,
-	TEXT("Whether to enable GPU direct for faster video frame copies. (Experimental)"),
-	ECVF_RenderThreadSafe);
-
-
 bool bBlackmagicWriteOutputRawDataCmdEnable = false;
 static FAutoConsoleCommand BlackmagicWriteOutputRawDataCmd(
 	TEXT("Blackmagic.WriteOutputRawData"),
@@ -298,25 +290,28 @@ namespace BlackmagicMediaPlayerHelpers
 							EncodeTime.Render(SetTimecode.Hours, SetTimecode.Minutes, SetTimecode.Seconds, SetTimecode.Frames);
 						}
 
-						const bool bUseGPUDirect = !!CVarBlackmagicEnableGPUDirect.GetValueOnAnyThread();
-						bool bGPUDirectTexturesAvailable = true;
-						if (bUseGPUDirect)
+						bool bGPUDirectTexturesAvailable = false;
+						if (MediaPlayer->CanUseGPUTextureTransfer())
 						{
 							if (!MediaPlayer->Textures.Num())
 							{
 								bGPUDirectTexturesAvailable = false;
 								UE_LOG(LogBlackmagicMedia, Error, TEXT("No texture available while doing a gpu texture transfer."));
 							}
+							else
+							{
+								bGPUDirectTexturesAvailable = true;
+							}
 						}
 
 						auto TextureSample = MediaPlayer->TextureSamplePool->AcquireShared();
 						bool bInitializeResult = false;
-						if (bGPUDirectTexturesAvailable && bUseGPUDirect)
+						if (bGPUDirectTexturesAvailable)
 						{
 							bInitializeResult = TextureSample->SetProperties(InFrameInfo.VideoPitch, InFrameInfo.VideoWidth, InFrameInfo.VideoHeight, SampleFormat, DecodedTime, MediaPlayer->VideoFrameRate, DecodedTimecode, bIsSRGBInput);
 							if (bInitializeResult)
 							{
-								TextureSample->SetBlackmagicBuffer(InFrameInfo.VideoBuffer);
+								TextureSample->SetBuffer(InFrameInfo.VideoBuffer);
 							}
 						}
 						else
@@ -335,45 +330,10 @@ namespace BlackmagicMediaPlayerHelpers
 
 						if (bInitializeResult)
 						{
-							if (bGPUDirectTexturesAvailable && bUseGPUDirect)
+							if (bGPUDirectTexturesAvailable && MediaPlayer->CanUseGPUTextureTransfer())
 							{
-								// Register the buffer and the texture
-								TRefCountPtr<FRHITexture> Texture = MediaPlayer->Textures.Pop();
-								if (Texture && MediaPlayer->RegisteredTextures.Contains(Texture))
-								{
-									TextureSample->SetTexture(Texture);
-									TextureSample->SetDestructionCallback([MediaPlayerWeakPtr = TWeakPtr<FBlackmagicMediaPlayer>(MediaPlayer->AsShared())](TRefCountPtr<FRHITexture> InTexture)
-									{
-										if (TSharedPtr<FBlackmagicMediaPlayer> MediaPlayerPtr = MediaPlayerWeakPtr.Pin())
-										{
-											MediaPlayerPtr->OnSampleDestroyed(InTexture);
-										}
-									});
-
-									MediaPlayer->RegisterSampleBuffer(TextureSample);
-								}
-								else
-								{
-									UE_LOG(LogBlackmagicMedia, Display, TEXT("Unregistered texture %u encountered while doing a gpu texture transfer."), Texture ? reinterpret_cast<std::uintptr_t>(Texture->GetNativeResource()) : 0);
-								}
-
-								ENQUEUE_RENDER_COMMAND(StartGPUTextureTransferCopy)(
-									[MediaPlayerWeakPtr = TWeakPtr<FBlackmagicMediaPlayer>(MediaPlayer->AsShared()), TextureSample](FRHICommandListImmediate& RHICmdList)
-									{
-										if (TextureSample->GetTexture())
-										{
-											if (TSharedPtr<FBlackmagicMediaPlayer> MediaPlayerPtr = MediaPlayerWeakPtr.Pin();
-												MediaPlayerPtr != nullptr && MediaPlayerPtr->GPUTextureTransfer != nullptr)
-											{ 
-												void* Buffer = TextureSample->GetBlackmagicBuffer();
-												UE_LOG(LogBlackmagicMedia, Verbose, TEXT("Starting Transfer with buffer %u"), reinterpret_cast<std::uintptr_t>(Buffer));
-												MediaPlayerPtr->GPUTextureTransfer->TransferTexture(Buffer, TextureSample->GetTexture(), UE::GPUTextureTransfer::ETransferDirection::CPU_TO_GPU);
-												MediaPlayerPtr->GPUTextureTransfer->BeginSync((TextureSample)->GetBlackmagicBuffer(), UE::GPUTextureTransfer::ETransferDirection::CPU_TO_GPU);
-												MediaPlayerPtr->Samples->AddVideo(TextureSample);
-												MediaPlayerPtr->GPUTextureTransfer->EndSync(TextureSample->GetBlackmagicBuffer());
-											}
-										}
-									});
+								MediaPlayer->PreGPUTransfer(TextureSample);
+								MediaPlayer->ExecuteGPUTransfer(TextureSample);
 							}
 							else
 							{
@@ -480,8 +440,6 @@ FBlackmagicMediaPlayer::FBlackmagicMediaPlayer(IMediaEventSink& InEventSink)
 	, bVerifyFrameDropCount(false)
 	, SupportedSampleTypes(EMediaIOSampleType::None)
 {
-	FGPUTextureTransferModule& Module = FGPUTextureTransferModule::Get();
-	GPUTextureTransfer = Module.GetTextureTransfer();
 }
 
 FBlackmagicMediaPlayer::~FBlackmagicMediaPlayer()
@@ -547,8 +505,9 @@ bool FBlackmagicMediaPlayer::Open(const FString& Url, const IMediaOptions* Optio
 	ChannelOptions.CallbackPriority = 10;
 	ChannelOptions.bReadVideo = Options->GetMediaOption(BlackmagicMediaOption::CaptureVideo, true);
 	ChannelOptions.FormatInfo.DisplayMode = Options->GetMediaOption(BlackmagicMediaOption::BlackmagicVideoFormat, (int64)BlackmagicMediaOption::DefaultVideoFormat);
-	ColorFormat = (EBlackmagicMediaSourceColorFormat)(Options->GetMediaOption(BlackmagicMediaOption::ColorFormat, (int64)EBlackmagicMediaSourceColorFormat::YUV8));
-	ChannelOptions.PixelFormat = ColorFormat == EBlackmagicMediaSourceColorFormat::YUV8 ? BlackmagicDesign::EPixelFormat::pf_8Bits : BlackmagicDesign::EPixelFormat::pf_10Bits;
+	BlackmagicColorFormat = (EBlackmagicMediaSourceColorFormat)(Options->GetMediaOption(BlackmagicMediaOption::ColorFormat, (int64)EBlackmagicMediaSourceColorFormat::YUV8));
+
+	ChannelOptions.PixelFormat = BlackmagicColorFormat == EBlackmagicMediaSourceColorFormat::YUV8 ? BlackmagicDesign::EPixelFormat::pf_8Bits : BlackmagicDesign::EPixelFormat::pf_10Bits;
 	const bool bIsSRGBInput = Options->GetMediaOption(BlackmagicMediaOption::SRGBInput, true);
 
 	const EMediaIOTimecodeFormat TimecodeFormat = (EMediaIOTimecodeFormat)(Options->GetMediaOption(BlackmagicMediaOption::TimecodeFormat, (int64)EMediaIOTimecodeFormat::None));
@@ -609,186 +568,9 @@ bool FBlackmagicMediaPlayer::Open(const FString& Url, const IMediaOptions* Optio
 	}
 #endif
 
-	if (bSuccess && CVarBlackmagicEnableGPUDirect.GetValueOnAnyThread())
-	{
-		CreateAndRegisterTextures(Options);
-	}
-
 	return bSuccess;
 }
 
-void FBlackmagicMediaPlayer::OnSampleDestroyed(TRefCountPtr<FRHITexture> InTexture)
-{
-	// Can be called after object destruction
-	FScopeLock ScopeLock(&TexturesCriticalSection);
-	if (InTexture && RegisteredTextures.Contains(InTexture))
-	{
-		Textures.Add(InTexture);
-	}
-}
-
-void FBlackmagicMediaPlayer::RegisterSampleBuffer(const TSharedPtr<FBlackmagicMediaTextureSample>& InSample)
-{
-	if (!RegisteredBuffers.Contains(InSample->GetBlackmagicBuffer()))
-	{
-		RegisteredBuffers.Add(InSample->GetBlackmagicBuffer());
-
-		// Enqueue buffer registration on the render  thread to ensure it happens before the call to TransferTexture
-		ENQUEUE_RENDER_COMMAND(RegisterSampleBuffers)([this, InSample](FRHICommandList& CommandList)
-			{
-				if (!InSample->GetBlackmagicBuffer())
-				{
-					UE_LOG(LogBlackmagicMedia, Error, TEXT("A buffer was not available while performing a gpu texture transfer."));
-					return;
-				}
-
-				if (!InSample->GetTexture())
-				{
-					UE_LOG(LogBlackmagicMedia, Error, TEXT("A texture was not available while performing a gpu texture transfer."));
-					return;
-				}
-
-				const uint32 TextureWidth = InSample->GetTexture()->GetDesc().Extent.X;
-				const uint32 TextureHeight = InSample->GetTexture()->GetDesc().Extent.Y;
-				uint32 TextureStride = TextureWidth * 4;
-
-				EPixelFormat Format = InSample->GetTexture()->GetFormat();
-				if (Format == PF_R32G32B32A32_UINT)
-				{
-					TextureStride *= 4;
-				}
-
-				UE_LOG(LogBlackmagicMedia, Verbose, TEXT("Registering buffer %u"), reinterpret_cast<std::uintptr_t>(InSample->GetBlackmagicBuffer()));
-				UE::GPUTextureTransfer::FRegisterDMABufferArgs Args;
-				Args.Buffer = InSample->GetBlackmagicBuffer();
-				Args.Width = TextureWidth;
-				Args.Height = TextureHeight;
-				Args.Stride = TextureStride;
-				Args.PixelFormat = Format == PF_B8G8R8A8 ? UE::GPUTextureTransfer::EPixelFormat::PF_8Bit : UE::GPUTextureTransfer::EPixelFormat::PF_10Bit;
-				GPUTextureTransfer->RegisterBuffer(Args);
-			});
-	}
-	else
-	{
-		UE_LOG(LogBlackmagicMedia, VeryVerbose, TEXT("Buffer already registered: %u"), reinterpret_cast<std::uintptr_t>(InSample->GetBlackmagicBuffer()));
-	}
-}
-
-void FBlackmagicMediaPlayer::UnregisterSampleBuffers()
-{
-	ENQUEUE_RENDER_COMMAND(UnregisterSampleBuffers)([TextureTransfer = GPUTextureTransfer, BuffersToUnregister = RegisteredBuffers](FRHICommandList& CommandList)
-		{
-			if (TextureTransfer)
-			{
-				for (void* RegisteredBuffer : BuffersToUnregister)
-				{
-					TextureTransfer->UnregisterBuffer(RegisteredBuffer);
-				}
-			}
-		});
-
-	RegisteredBuffers.Reset();
-}
-
-void FBlackmagicMediaPlayer::CreateAndRegisterTextures(const IMediaOptions * Options)
-{
-	ENQUEUE_RENDER_COMMAND(BlackmagicMediaPlayerCreateTextures)(
-		[this, Options](FRHICommandListImmediate& CommandList)
-		{
-			if (!GPUTextureTransfer)
-			{
-				UE_LOG(LogBlackmagicMedia, Error, TEXT("Could not register textures for gpu texture transfer, as the GPUTextureTransfer object was null."));
-				return;
-			}
-
-			for (uint8 Index = 0; Index < MaxNumVideoFrameBuffer; Index++)
-			{
-				TRefCountPtr<FRHITexture> RHITexture;
-				const FString TextureName = FString::Printf(TEXT("FBlackmagicMediaPlayerTexture %d"), Index);
-				FRHIResourceCreateInfo CreateInfo(*TextureName);
-
-				uint32 TextureWidth = VideoTrackFormat.Dim.X;
-				const uint32 TextureHeight = VideoTrackFormat.Dim.Y;
-				uint32 Stride = TextureWidth;
-				EPixelFormat InputFormat = EPixelFormat::PF_Unknown;
-
-				if (ColorFormat == EBlackmagicMediaSourceColorFormat::YUV8)
-				{
-					InputFormat = EPixelFormat::PF_B8G8R8A8;
-					TextureWidth /= 2;
-					Stride = TextureWidth * 4;
-				}
-				else if (ColorFormat == EBlackmagicMediaSourceColorFormat::YUV10)
-				{
-					InputFormat = EPixelFormat::PF_R32G32B32A32_UINT;
-					TextureWidth /= 6;
-					Stride = TextureWidth * 16;
-				}
-				else
-				{
-					checkf(false, TEXT("Format not supported"));
-				}
-
-				TRefCountPtr<FRHITexture2D> DummyTexture2DRHI;
-
-				ETextureCreateFlags CreateFlags = TexCreate_Shared;
-				if (RHIGetInterfaceType() == ERHIInterfaceType::Vulkan)
-				{
-					CreateFlags = TexCreate_External;
-				}
-
-				FRHITextureCreateDesc CreateDesc = FRHITextureCreateDesc::Create2D(*TextureName, FIntPoint(TextureWidth, TextureHeight), InputFormat, FClearValueBinding::White, CreateFlags);
-				RHICreateTargetableShaderResource(CreateDesc, TexCreate_RenderTargetable, RHITexture, DummyTexture2DRHI);
-
-				UE_LOG(LogBlackmagicMedia, Verbose, TEXT("Registering texture %u"), reinterpret_cast<std::uintptr_t>(RHITexture->GetNativeResource()));
-
-				UE::GPUTextureTransfer::FRegisterDMATextureArgs Args;
-				Args.RHITexture = RHITexture.GetReference();
-				Args.Stride = Stride;
-				Args.Width = TextureWidth;
-				Args.Height = TextureHeight;
-				Args.RHIResourceMemory = nullptr;
-				Args.PixelFormat = ColorFormat == EBlackmagicMediaSourceColorFormat::YUV8 ? UE::GPUTextureTransfer::EPixelFormat::PF_8Bit : UE::GPUTextureTransfer::EPixelFormat::PF_10Bit;
-
-				GPUTextureTransfer->RegisterTexture(Args);
-				RegisteredTextures.Add(RHITexture);
-
-				{
-					FScopeLock Lock(&TexturesCriticalSection);
-					Textures.Add(RHITexture);
-				}
-			}
-		});
-
-	FRenderCommandFence RenderFence;
-	RenderFence.BeginFence();
-	RenderFence.Wait();
-}
-
-void FBlackmagicMediaPlayer::UnregisterTextures()
-{
-	// Keep a copy of the textures to make sure they were not destoryed before this gets executed.
-	ENQUEUE_RENDER_COMMAND(UnregisterSampleBuffers)([TextureTransfer = GPUTextureTransfer, TexturesToUnregister = RegisteredTextures](FRHICommandList& CommandList)
-		{
-			if (TextureTransfer)
-			{
-				for (const TRefCountPtr<FRHITexture>& RHITexture : TexturesToUnregister)
-				{
-					UE_LOG(LogBlackmagicMedia, Verbose, TEXT("Unregistering texture %u"), reinterpret_cast<std::uintptr_t>(RHITexture->GetNativeResource()));
-					TextureTransfer->UnregisterTexture(RHITexture.GetReference());
-				}
-			}
-		});
-
-	RegisteredTextures.Reset();
-
-	{
-		FScopeLock Lock(&TexturesCriticalSection);
-		Textures.Reset();
-	}
-}
-
-
 void FBlackmagicMediaPlayer::TickInput(FTimespan DeltaTime, FTimespan Timecode)
 {
 	// update player state
@@ -861,5 +643,10 @@ void FBlackmagicMediaPlayer::SetupSampleChannels()
 	Samples->InitializeAudioBuffer(AudioSettings);
 }
 
+void FBlackmagicMediaPlayer::AddVideoSample(const TSharedRef<FMediaIOCoreTextureSampleBase>& InSample)
+{
+	Samples->AddVideo(InSample);
+}
+
 #undef LOCTEXT_NAMESPACE
 
diff --git a/Engine/Plugins/Media/BlackmagicMedia/Source/BlackmagicMedia/Private/Player/BlackmagicMediaPlayer.h b/Engine/Plugins/Media/BlackmagicMedia/Source/BlackmagicMedia/Private/Player/BlackmagicMediaPlayer.h
--- a/Engine/Plugins/Media/BlackmagicMedia/Source/BlackmagicMedia/Private/Player/BlackmagicMediaPlayer.h
+++ b/Engine/Plugins/Media/BlackmagicMedia/Source/BlackmagicMedia/Private/Player/BlackmagicMediaPlayer.h
@@ -5,8 +5,6 @@
 #include "MediaIOCorePlayerBase.h"
 
 #include "BlackmagicMediaSource.h"
-#include "GPUTextureTransferModule.h"
-#include "GPUTextureTransfer.h"
 #include "HAL/CriticalSection.h"
 #include "MediaIOCoreAudioSampleBase.h"
 #include "MediaIOCoreTextureSampleBase.h"
@@ -27,67 +25,6 @@ class FBlackmagicMediaTextureSample : public FMediaIOCoreTextureSampleBase
 {
 public:
 	virtual const FMatrix& GetYUVToRGBMatrix() const override { return MediaShaders::YuvToRgbRec709Scaled; }
-
-	virtual void ShutdownPoolable() override
-	{
-		if (DestructionCallback)
-		{
-			DestructionCallback(Texture);
-		}
-
-		FMediaIOCoreTextureSampleBase::FreeSample();
-	}
-
-	void SetTexture(TRefCountPtr<FRHITexture> InRHITexture)
-	{
-		Texture = MoveTemp(InRHITexture);
-	}
-
-	void SetDestructionCallback(TFunction<void(TRefCountPtr<FRHITexture2D>)> InDestructionCallback)
-	{
-		DestructionCallback = InDestructionCallback;
-	}
-
-	virtual const void* GetBuffer() override
-	{
-		// Don't return the buffer if we have a texture to force the media player to use the texture if available. 
-		if (Texture)
-		{
-			return nullptr;
-		}
-		return Buffer.GetData();
-
-	}
-
-	void* GetBlackmagicBuffer()
-	{
-		return BlackmagicBufferPtr;
-	}
-
-	void SetBlackmagicBuffer(void* InBuffer)
-	{
-		BlackmagicBufferPtr = InBuffer;
-	}
-
-#if WITH_ENGINE
-	virtual FRHITexture* GetTexture() const override
-	{
-		if (Texture)
-		{
-			return Texture.GetReference();
-		}
-		return nullptr;
-	}
-#endif //WITH_ENGINE
-
-private:
-	/** Hold a texture to be used for gpu texture transfers. */
-	TRefCountPtr<FRHITexture2D> Texture;
-
-	/** Called when the sample is destroyed by its pool. */
-	TFunction<void(TRefCountPtr<FRHITexture2D>)> DestructionCallback;
-
-	void* BlackmagicBufferPtr = nullptr;
 };
 
 class FBlackmagicMediaAudioSamplePool : public TMediaObjectPool<FMediaIOCoreAudioSampleBase> { };
@@ -104,7 +41,7 @@ class FBlackmagicMediaTextureSamplePool : public TMediaObjectPool<FBlackmagicMed
  * the player's current play time (CurrentTime) is derived either from the
  * time codes embedded in frames or from the Engine's global time code.
  */
-class FBlackmagicMediaPlayer : public FMediaIOCorePlayerBase, public TSharedFromThis<FBlackmagicMediaPlayer>
+class FBlackmagicMediaPlayer : public FMediaIOCorePlayerBase
 {
 private:
 	using Super = FMediaIOCorePlayerBase;
@@ -152,13 +89,18 @@ public:
 protected:
 	/** Setup our different channels with the current set of settings */
 	virtual void SetupSampleChannels() override;
-	void OnSampleDestroyed(TRefCountPtr<FRHITexture> InTexture);
-	void RegisterSampleBuffer(const TSharedPtr<FBlackmagicMediaTextureSample>& InSample);
-	void UnregisterSampleBuffers();
-	void CreateAndRegisterTextures(const IMediaOptions* Options);
-	void UnregisterTextures();
-private:
 
+	virtual uint32 GetNumVideoFrameBuffers() const override 
+	{
+		return MaxNumVideoFrameBuffer;
+	}
+
+	virtual EMediaIOCoreColorFormat GetColorFormat() const override
+	{
+		return BlackmagicColorFormat == EBlackmagicMediaSourceColorFormat::YUV8 ? EMediaIOCoreColorFormat::YUV8 : EMediaIOCoreColorFormat::YUV10;
+	}
+
+	virtual void AddVideoSample(const TSharedRef<FMediaIOCoreTextureSampleBase>& InSample) override;
 
 private:
 
@@ -179,17 +121,5 @@ private:
 	/** Used to flag which sample types we advertise as supported for timed data monitoring */
 	EMediaIOSampleType SupportedSampleTypes;
 
-	FCriticalSection TexturesCriticalSection;
-
-	/** GPU Texture transfer object */
-	UE::GPUTextureTransfer::TextureTransferPtr GPUTextureTransfer;
-
-	/** Pool of textures registerd with GPU Texture transfer. */
-	TArray<TRefCountPtr<FRHITexture>> Textures;
-	/** Buffers Registered with GPU Texture Transfer */
-	TSet<void*> RegisteredBuffers;
-	/** Pool of textures registerd with GPU Texture transfer. */
-	TSet<TRefCountPtr<FRHITexture>> RegisteredTextures;
-
-	EBlackmagicMediaSourceColorFormat ColorFormat = EBlackmagicMediaSourceColorFormat::YUV8;
+	EBlackmagicMediaSourceColorFormat BlackmagicColorFormat = EBlackmagicMediaSourceColorFormat::YUV8;
 };
diff --git a/Engine/Plugins/Media/MediaIOFramework/MediaIOFramework.uplugin b/Engine/Plugins/Media/MediaIOFramework/MediaIOFramework.uplugin
--- a/Engine/Plugins/Media/MediaIOFramework/MediaIOFramework.uplugin
+++ b/Engine/Plugins/Media/MediaIOFramework/MediaIOFramework.uplugin
@@ -27,8 +27,8 @@
 		{
 			"Name": "GPUTextureTransfer",
 			"Type": "Runtime",
-			"LoadingPhase": "PostConfigInit", 
-			"PlatformAllowList": [ "Win64", "Linux" ]
+			"LoadingPhase": "PostConfigInit",
+			"PlatformAllowList": [ "Win64", "Mac", "Linux" ]
 		}
 	],
 	"DocsURL": "",
diff --git a/Engine/Plugins/Media/MediaIOFramework/Source/GPUTextureTransfer/Private/GPUTextureTransferModule.cpp b/Engine/Plugins/Media/MediaIOFramework/Source/GPUTextureTransfer/Private/GPUTextureTransferModule.cpp
--- a/Engine/Plugins/Media/MediaIOFramework/Source/GPUTextureTransfer/Private/GPUTextureTransferModule.cpp
+++ b/Engine/Plugins/Media/MediaIOFramework/Source/GPUTextureTransfer/Private/GPUTextureTransferModule.cpp
@@ -9,8 +9,17 @@
 #include "TextureTransferBase.h"
 #endif
 
-#include "CoreMinimal.h"
+#if PLATFORM_WINDOWS || PLATFORM_MAC || PLATFORM_LINUX
+#define VULKAN_PLATFORM 1
+#else
+#define VULKAN_PLATFORM 0
+#endif
+
+#if VULKAN_PLATFORM
 #include "IVulkanDynamicRHI.h"
+#endif
+
+#include "CoreMinimal.h"
 #include "GenericPlatform/GenericPlatformDriver.h"
 #include "HAL/PlatformMisc.h"
 #include "Misc/App.h"
@@ -57,7 +66,10 @@ void FGPUTextureTransferModule::StartupModule()
 #elif PLATFORM_LINUX
 				const TArray<const ANSICHAR*> ExtentionsToAdd{ VK_KHR_EXTERNAL_MEMORY_EXTENSION_NAME, VK_KHR_EXTERNAL_MEMORY_FD_EXTENSION_NAME, VK_KHR_EXTERNAL_SEMAPHORE_FD_EXTENSION_NAME, VK_KHR_SURFACE_EXTENSION_NAME };
 #endif
+
+#if VULKAN_PLATFORM
 				IVulkanDynamicRHI::AddEnabledDeviceExtensionsAndLayers(ExtentionsToAdd, TArray<const ANSICHAR*>());
+#endif
 			}
 
 			TransferObjects.AddDefaulted(RHI_MAX);
@@ -178,12 +190,14 @@ void FGPUTextureTransferModule::InitializeTextureTransfer()
 		InitializeArgs.RHI = RHI;
 		InitializeArgs.RHIDevice = GDynamicRHI->RHIGetNativeDevice();
 		InitializeArgs.RHICommandQueue = GDynamicRHI->RHIGetNativeGraphicsQueue();
+#if VULKAN_PLATFORM
 		if (RHI == UE::GPUTextureTransfer::ERHI::Vulkan)
 		{
 			IVulkanDynamicRHI* DynRHI = GetIVulkanDynamicRHI();
 			InitializeArgs.VulkanInstance = DynRHI->RHIGetVkInstance();
 			FMemory::Memcpy(InitializeArgs.RHIDeviceUUID, DynRHI->RHIGetVulkanDeviceUUID(), 16);
 		}
+#endif
 
 		const uint8 RHIIndex = static_cast<uint8>(RHI);
 		if (TextureTransfer->Initialize(InitializeArgs))
diff --git a/Engine/Plugins/Media/MediaIOFramework/Source/MediaIOCore/MediaIOCore.Build.cs b/Engine/Plugins/Media/MediaIOFramework/Source/MediaIOCore/MediaIOCore.Build.cs
--- a/Engine/Plugins/Media/MediaIOFramework/Source/MediaIOCore/MediaIOCore.Build.cs
+++ b/Engine/Plugins/Media/MediaIOFramework/Source/MediaIOCore/MediaIOCore.Build.cs
@@ -25,12 +25,22 @@ namespace UnrealBuildTool.Rules
 					"SlateCore",
 					"TimeManagement"
 				});
-
+			
 			PrivateIncludePaths.AddRange(
 				new string[] {
 				});
 
-            PrivateDependencyModuleNames.AddRange(
+			if (Target.Platform == UnrealTargetPlatform.Win64 || Target.Platform == UnrealTargetPlatform.Linux ||
+			    Target.Platform == UnrealTargetPlatform.Mac)
+			{
+				PrivateDependencyModuleNames.AddRange(
+					new string[]
+					{
+						"GPUTextureTransfer"
+					});
+			}
+
+			PrivateDependencyModuleNames.AddRange(
                 new string[]
                 {
                     "AudioMixer",
diff --git a/Engine/Plugins/Media/MediaIOFramework/Source/MediaIOCore/Private/Player/MediaIOCorePlayerBase.cpp b/Engine/Plugins/Media/MediaIOFramework/Source/MediaIOCore/Private/Player/MediaIOCorePlayerBase.cpp
--- a/Engine/Plugins/Media/MediaIOFramework/Source/MediaIOCore/Private/Player/MediaIOCorePlayerBase.cpp
+++ b/Engine/Plugins/Media/MediaIOFramework/Source/MediaIOCore/Private/Player/MediaIOCorePlayerBase.cpp
@@ -14,9 +14,13 @@
 #include "MediaIOCoreSamples.h"
 #include "Misc/App.h"
 #include "Misc/ScopeLock.h"
+#include "PixelFormat.h"
 #include "TimedDataInputCollection.h"
 #include "TimeSynchronizableMediaSource.h"
 
+#if PLATFORM_DESKTOP
+#include "GPUTextureTransferModule.h"
+#endif
 
 #define LOCTEXT_NAMESPACE "MediaIOCorePlayerBase"
 
@@ -55,6 +59,12 @@ namespace MediaIOCorePlayerDetail
 		TEXT("All media player will stop logging the frame timecode when a new frame is captured."),
 		FConsoleCommandDelegate::CreateLambda([]() { MediaIOCorePlayerDetail::bLogTimecode = false; }),
 		ECVF_Cheat);
+
+	static TAutoConsoleVariable<int32> CVarEnableGPUDirectInput(
+		TEXT("MediaIO.EnableGPUDirectInput"), 0,
+		TEXT("Whether to enable GPU direct for faster video frame copies. (Experimental)"),
+		ECVF_RenderThreadSafe);
+
 }
 
 /* FMediaIOCoreMediaOption structors
@@ -81,6 +91,11 @@ FMediaIOCorePlayerBase::FMediaIOCorePlayerBase(IMediaEventSink& InEventSink)
 	, PreviousFrameTimespan(FTimespan::Zero())
 {
 	Samples = new FMediaIOCoreSamples();
+
+#if PLATFORM_DESKTOP
+	FGPUTextureTransferModule& Module = FGPUTextureTransferModule::Get();
+	GPUTextureTransfer = Module.GetTextureTransfer();
+#endif
 }
 
 FMediaIOCorePlayerBase::~FMediaIOCorePlayerBase()
@@ -183,7 +198,14 @@ bool FMediaIOCorePlayerBase::Open(const FString& Url, const IMediaOptions* Optio
 	ITimeManagementModule::Get().GetTimedDataInputCollection().Add(this);
 
 	OpenUrl = Url;
-	return ReadMediaOptions(Options);
+	bool bReadMediaOptions = ReadMediaOptions(Options);
+
+	if (MediaIOCorePlayerDetail::CVarEnableGPUDirectInput.GetValueOnAnyThread())
+	{
+		CreateAndRegisterTextures(Options);
+	}
+
+	return bReadMediaOptions;
 }
 
 bool FMediaIOCorePlayerBase::Open(const TSharedRef<FArchive, ESPMode::ThreadSafe>& /*Archive*/, const FString& /*OriginalUrl*/, const IMediaOptions* /*Options*/)
@@ -498,6 +520,243 @@ bool FMediaIOCorePlayerBase::IsTimecodeLogEnabled()
 #endif
 }
 
+bool FMediaIOCorePlayerBase::CanUseGPUTextureTransfer()
+{
+#if PLATFORM_DESKTOP
+	return GPUTextureTransfer && MediaIOCorePlayerDetail::CVarEnableGPUDirectInput.GetValueOnAnyThread();
+#else
+	return false;
+#endif
+}
+
+void FMediaIOCorePlayerBase::OnSampleDestroyed(TRefCountPtr<FRHITexture> InTexture)
+{
+	FScopeLock ScopeLock(&TexturesCriticalSection);
+	if (InTexture && RegisteredTextures.Contains(InTexture))
+	{
+		Textures.Add(InTexture);
+	}
+}
+
+void FMediaIOCorePlayerBase::RegisterSampleBuffer(const TSharedPtr<FMediaIOCoreTextureSampleBase>& InSample)
+{
+#if PLATFORM_DESKTOP
+	if (!RegisteredBuffers.Contains(InSample->GetMutableBuffer()))
+	{
+		RegisteredBuffers.Add(InSample->GetMutableBuffer());
+
+		// Enqueue buffer registration on the render  thread to ensure it happens before the call to TransferTexture
+		ENQUEUE_RENDER_COMMAND(RegisterSampleBuffers)([this, InSample](FRHICommandList& CommandList)
+			{
+				if (!InSample->GetMutableBuffer())
+				{
+					UE_LOG(LogMediaIOCore, Error, TEXT("A buffer was not available while performing a gpu texture transfer."));
+					return;
+				}
+
+				if (!InSample->GetTexture())
+				{
+					UE_LOG(LogMediaIOCore, Error, TEXT("A texture was not available while performing a gpu texture transfer."));
+					return;
+				}
+
+				const uint32 TextureWidth = InSample->GetTexture()->GetDesc().Extent.X;
+				const uint32 TextureHeight = InSample->GetTexture()->GetDesc().Extent.Y;
+				uint32 TextureStride = TextureWidth * 4;
+
+				EPixelFormat Format = InSample->GetTexture()->GetFormat();
+				if (Format == PF_R32G32B32A32_UINT)
+				{
+					TextureStride *= 4;
+				}
+
+				UE_LOG(LogMediaIOCore, Verbose, TEXT("Registering buffer %u"), reinterpret_cast<std::uintptr_t>(InSample->GetMutableBuffer()));
+				UE::GPUTextureTransfer::FRegisterDMABufferArgs Args;
+				Args.Buffer = InSample->GetMutableBuffer();
+				Args.Width = TextureWidth;
+				Args.Height = TextureHeight;
+				Args.Stride = TextureStride;
+				Args.PixelFormat = Format == PF_B8G8R8A8 ? UE::GPUTextureTransfer::EPixelFormat::PF_8Bit : UE::GPUTextureTransfer::EPixelFormat::PF_10Bit;
+				GPUTextureTransfer->RegisterBuffer(Args);
+			});
+	}
+	else
+	{
+		UE_LOG(LogMediaIOCore, VeryVerbose, TEXT("Buffer already registered: %u"), reinterpret_cast<std::uintptr_t>(InSample->GetMutableBuffer()));
+	}
+#endif
+}
+
+void FMediaIOCorePlayerBase::UnregisterSampleBuffers()
+{
+#if PLATFORM_DESKTOP
+	ENQUEUE_RENDER_COMMAND(UnregisterSampleBuffers)([TextureTransfer = GPUTextureTransfer, BuffersToUnregister = RegisteredBuffers](FRHICommandList& CommandList)
+		{
+			if (TextureTransfer)
+			{
+				for (void* RegisteredBuffer : BuffersToUnregister)
+				{
+					TextureTransfer->UnregisterBuffer(RegisteredBuffer);
+				}
+			}
+		});
+
+	RegisteredBuffers.Reset();
+#endif
+}
+
+void FMediaIOCorePlayerBase::CreateAndRegisterTextures(const IMediaOptions* Options)
+{
+#if PLATFORM_DESKTOP
+	ENQUEUE_RENDER_COMMAND(MediaIOCorePlayerCreateTextures)(
+		[this, Options](FRHICommandListImmediate& CommandList)
+		{
+			if (!GPUTextureTransfer)
+			{
+				UE_LOG(LogMediaIOCore, Error, TEXT("Could not register textures for gpu texture transfer, as the GPUTextureTransfer object was null."));
+				return;
+			}
+
+			for (uint8 Index = 0; Index < GetNumVideoFrameBuffers(); Index++)
+			{
+				TRefCountPtr<FRHITexture> RHITexture;
+				const FString TextureName = FString::Printf(TEXT("FMediaIOCorePlayerTexture %d"), Index);
+				FRHIResourceCreateInfo CreateInfo(*TextureName);
+
+				uint32 TextureWidth = VideoTrackFormat.Dim.X;
+				const uint32 TextureHeight = VideoTrackFormat.Dim.Y;
+				uint32 Stride = TextureWidth;
+				EPixelFormat InputFormat = EPixelFormat::PF_Unknown;
+
+				EMediaIOCoreColorFormat ColorFormat = GetColorFormat();
+
+				if (ColorFormat == EMediaIOCoreColorFormat::YUV8)
+				{
+					InputFormat = EPixelFormat::PF_B8G8R8A8;
+					TextureWidth /= 2;
+					Stride = TextureWidth * 4;
+				}
+				else if (ColorFormat == EMediaIOCoreColorFormat::YUV10)
+				{
+					InputFormat = EPixelFormat::PF_R32G32B32A32_UINT;
+					TextureWidth /= 6;
+					Stride = TextureWidth * 16;
+				}
+				else
+				{
+					checkf(false, TEXT("Format not supported"));
+				}
+
+				TRefCountPtr<FRHITexture2D> DummyTexture2DRHI;
+
+				ETextureCreateFlags CreateFlags = TexCreate_Shared;
+				if (RHIGetInterfaceType() == ERHIInterfaceType::Vulkan)
+				{
+					CreateFlags = TexCreate_External;
+				}
+
+				FRHITextureCreateDesc CreateDesc = FRHITextureCreateDesc::Create2D(*TextureName, FIntPoint(TextureWidth, TextureHeight), InputFormat, FClearValueBinding::White, CreateFlags);
+				RHICreateTargetableShaderResource(CreateDesc, TexCreate_RenderTargetable, RHITexture, DummyTexture2DRHI);
+
+				UE_LOG(LogMediaIOCore, Verbose, TEXT("Registering texture %u"), reinterpret_cast<std::uintptr_t>(RHITexture->GetNativeResource()));
+
+				UE::GPUTextureTransfer::FRegisterDMATextureArgs Args;
+				Args.RHITexture = RHITexture.GetReference();
+				Args.Stride = Stride;
+				Args.Width = TextureWidth;
+				Args.Height = TextureHeight;
+				Args.RHIResourceMemory = nullptr;
+				Args.PixelFormat = ColorFormat == EMediaIOCoreColorFormat::YUV8 ? UE::GPUTextureTransfer::EPixelFormat::PF_8Bit : UE::GPUTextureTransfer::EPixelFormat::PF_10Bit;
+
+				GPUTextureTransfer->RegisterTexture(Args);
+				RegisteredTextures.Add(RHITexture);
+
+				{
+					FScopeLock Lock(&TexturesCriticalSection);
+					Textures.Add(RHITexture);
+				}
+			}
+		});
+
+	FRenderCommandFence RenderFence;
+	RenderFence.BeginFence();
+	RenderFence.Wait();
+#endif
+}
+
+void FMediaIOCorePlayerBase::UnregisterTextures()
+{
+#if PLATFORM_DESKTOP
+	// Keep a copy of the textures to make sure they were not destoryed before this gets executed.
+	ENQUEUE_RENDER_COMMAND(UnregisterSampleBuffers)([TextureTransfer = GPUTextureTransfer, TexturesToUnregister = RegisteredTextures](FRHICommandList& CommandList)
+		{
+			if (TextureTransfer)
+			{
+				for (const TRefCountPtr<FRHITexture>& RHITexture : TexturesToUnregister)
+				{
+					UE_LOG(LogMediaIOCore, Verbose, TEXT("Unregistering texture %u"), reinterpret_cast<std::uintptr_t>(RHITexture->GetNativeResource()));
+					TextureTransfer->UnregisterTexture(RHITexture.GetReference());
+				}
+			}
+		});
+
+	RegisteredTextures.Reset();
+
+	{
+		FScopeLock Lock(&TexturesCriticalSection);
+		Textures.Reset();
+	}
+#endif
+}
+
+void FMediaIOCorePlayerBase::PreGPUTransfer(const TSharedPtr<FMediaIOCoreTextureSampleBase>& InSample)
+{
+#if PLATFORM_DESKTOP
+	// Register the buffer and the texture
+	TRefCountPtr<FRHITexture> Texture = Textures.Pop();
+	if (Texture && RegisteredTextures.Contains(Texture))
+	{
+		InSample->SetTexture(Texture);
+		InSample->SetDestructionCallback([MediaPlayerWeakPtr = TWeakPtr<FMediaIOCorePlayerBase>(AsShared())](TRefCountPtr<FRHITexture> InTexture)
+		{
+			if (TSharedPtr<FMediaIOCorePlayerBase> MediaPlayerPtr = MediaPlayerWeakPtr.Pin())
+			{
+				MediaPlayerPtr->OnSampleDestroyed(InTexture);
+			}
+		});
+
+		RegisterSampleBuffer(InSample);
+	}
+	else
+	{
+		UE_LOG(LogMediaIOCore, Display, TEXT("Unregistered texture %u encountered while doing a gpu texture transfer."), Texture ? reinterpret_cast<std::uintptr_t>(Texture->GetNativeResource()) : 0);
+	}
+#endif
+}
+
+void FMediaIOCorePlayerBase::ExecuteGPUTransfer(const TSharedPtr<FMediaIOCoreTextureSampleBase>& InSample)
+{
+#if PLATFORM_DESKTOP
+	ENQUEUE_RENDER_COMMAND(StartGPUTextureTransferCopy)(
+		[MediaPlayerWeakPtr = TWeakPtr<FMediaIOCorePlayerBase>(AsShared()), InSample](FRHICommandListImmediate& RHICmdList)
+	{
+		if (InSample->GetTexture())
+		{
+			if (TSharedPtr<FMediaIOCorePlayerBase> MediaPlayerPtr = MediaPlayerWeakPtr.Pin();
+				MediaPlayerPtr != nullptr && MediaPlayerPtr->GPUTextureTransfer != nullptr)
+			{
+				void* Buffer = InSample->GetMutableBuffer();
+				UE_LOG(LogMediaIOCore, Verbose, TEXT("Starting Transfer with buffer %u"), reinterpret_cast<std::uintptr_t>(Buffer));
+				MediaPlayerPtr->GPUTextureTransfer->TransferTexture(Buffer, InSample->GetTexture(), UE::GPUTextureTransfer::ETransferDirection::CPU_TO_GPU);
+				MediaPlayerPtr->GPUTextureTransfer->BeginSync((InSample)->GetMutableBuffer(), UE::GPUTextureTransfer::ETransferDirection::CPU_TO_GPU);
+				MediaPlayerPtr->AddVideoSample(InSample.ToSharedRef());
+				MediaPlayerPtr->GPUTextureTransfer->EndSync(InSample->GetMutableBuffer());
+			}
+		}
+	});
+#endif
+}
+
 TArray<ITimedDataInputChannel*> FMediaIOCorePlayerBase::GetChannels() const
 {
 	return Channels;
diff --git a/Engine/Plugins/Media/MediaIOFramework/Source/MediaIOCore/Public/MediaIOCorePlayerBase.h b/Engine/Plugins/Media/MediaIOFramework/Source/MediaIOCore/Public/MediaIOCorePlayerBase.h
--- a/Engine/Plugins/Media/MediaIOFramework/Source/MediaIOCore/Public/MediaIOCorePlayerBase.h
+++ b/Engine/Plugins/Media/MediaIOFramework/Source/MediaIOCore/Public/MediaIOCorePlayerBase.h
@@ -13,6 +13,7 @@
 
 #include "HAL/CriticalSection.h"
 #include "MediaIOCoreSampleContainer.h"
+#include "MediaIOCoreTextureSampleBase.h"
 #include "Misc/CoreMisc.h"
 
 #include "Misc/FrameRate.h"
@@ -22,6 +23,12 @@ class IMediaEventSink;
 
 enum class EMediaTextureSampleFormat;
 
+enum class EMediaIOCoreColorFormat : uint8
+{
+	YUV8,
+	YUV10
+};
+
 struct MEDIAIOCORE_API FMediaIOCoreMediaOption
 {
 	static const FName FrameRateNumerator;
@@ -31,6 +38,11 @@ struct MEDIAIOCORE_API FMediaIOCoreMediaOption
 	static const FName VideoModeName;
 };
 
+namespace UE::GPUTextureTransfer
+{
+	using TextureTransferPtr = TSharedPtr<class ITextureTransfer>;
+}
+
 /**
  * Implements a base player for hardware IO cards. 
  *
@@ -49,6 +61,7 @@ class MEDIAIOCORE_API FMediaIOCorePlayerBase
 	, protected IMediaTracks
 	, protected IMediaView
 	, public ITimedDataInput
+	, public TSharedFromThis<FMediaIOCorePlayerBase>
 {
 public:
 
@@ -155,6 +168,33 @@ protected:
 	 */
 	virtual void SetupSampleChannels() = 0;
 
+	/**
+	 * Get the number of video frames to buffer.
+	 */
+	virtual uint32 GetNumVideoFrameBuffers() const
+	{ 
+		return 1;
+	}
+
+	virtual EMediaIOCoreColorFormat GetColorFormat() const 
+	{ 
+		return EMediaIOCoreColorFormat::YUV8;
+	}
+
+	virtual void AddVideoSample(const TSharedRef<FMediaIOCoreTextureSampleBase>& InSample)
+	{
+	}
+
+	virtual bool CanUseGPUTextureTransfer();
+
+	void OnSampleDestroyed(TRefCountPtr<FRHITexture> InTexture);
+	void RegisterSampleBuffer(const TSharedPtr<FMediaIOCoreTextureSampleBase>& InSample);
+	void UnregisterSampleBuffers();
+	void CreateAndRegisterTextures(const IMediaOptions* Options);
+	void UnregisterTextures();
+	void PreGPUTransfer(const TSharedPtr<FMediaIOCoreTextureSampleBase>& InSample);
+	void ExecuteGPUTransfer(const TSharedPtr<FMediaIOCoreTextureSampleBase>& InSample);
+
 protected:
 	/** Critical section for synchronizing access to receiver and sinks. */
 	FCriticalSection CriticalSection;
@@ -203,6 +243,17 @@ protected:
 	
 	/** Base set of settings to start from when setuping channels */
 	FMediaIOSamplingSettings BaseSettings;
-};
 
+	FCriticalSection TexturesCriticalSection;
 
+	/** Pool of textures registerd with GPU Texture transfer. */
+	TArray<TRefCountPtr<FRHITexture>> Textures;
+
+private:
+	/** GPU Texture transfer object */
+	UE::GPUTextureTransfer::TextureTransferPtr GPUTextureTransfer;
+	/** Buffers Registered with GPU Texture Transfer */
+	TSet<void*> RegisteredBuffers;
+	/** Pool of textures registerd with GPU Texture transfer. */
+	TSet<TRefCountPtr<FRHITexture>> RegisteredTextures;
+};
diff --git a/Engine/Plugins/Media/MediaIOFramework/Source/MediaIOCore/Public/MediaIOCoreTextureSampleBase.h b/Engine/Plugins/Media/MediaIOFramework/Source/MediaIOCore/Public/MediaIOCoreTextureSampleBase.h
--- a/Engine/Plugins/Media/MediaIOFramework/Source/MediaIOCore/Public/MediaIOCoreTextureSampleBase.h
+++ b/Engine/Plugins/Media/MediaIOFramework/Source/MediaIOCore/Public/MediaIOCoreTextureSampleBase.h
@@ -114,14 +114,9 @@ public:
 public:
 	//~ IMediaTextureSample interface
 
-	virtual const void* GetBuffer() override
-	{
-		return Buffer.GetData();
-	}
-
 	virtual FIntPoint GetDim() const override
 	{
-		switch(GetFormat())
+		switch (GetFormat())
 		{
 		case EMediaTextureSampleFormat::CharAYUV:
 		case EMediaTextureSampleFormat::CharNV12:
@@ -158,13 +153,6 @@ public:
 		return Stride;
 	}
 
-#if WITH_ENGINE
-	virtual FRHITexture* GetTexture() const override
-	{
-		return nullptr;
-	}
-#endif //WITH_ENGINE
-
 	virtual FMediaTimeStamp GetTime() const override
 	{
 		return FMediaTimeStamp(Time);
@@ -185,11 +173,73 @@ public:
 		return bIsSRGBInput;
 	}
 
+
+	virtual const void* GetBuffer() override
+	{
+		// Don't return the buffer if we have a texture to force the media player to use the texture if available. 
+		if (Texture)
+		{
+			return nullptr;
+		}
+
+		if (ExternalBuffer)
+		{
+			return ExternalBuffer;
+		}
+
+		return Buffer.GetData();
+
+	}
+
+	void* GetMutableBuffer()
+	{
+		if (ExternalBuffer)
+		{
+			return ExternalBuffer;
+		}
+
+		return Buffer.GetData();
+	}
+
+#if WITH_ENGINE
+	virtual FRHITexture* GetTexture() const override
+	{
+		return Texture.GetReference();
+	}
+#endif //WITH_ENGINE
+
+	void SetBuffer(void* InBuffer)
+	{
+		ExternalBuffer = InBuffer;
+	}
+
+	void SetTexture(TRefCountPtr<FRHITexture> InRHITexture)
+	{
+		Texture = MoveTemp(InRHITexture);
+	}
+
+	void SetDestructionCallback(TFunction<void(TRefCountPtr<FRHITexture2D>)> InDestructionCallback)
+	{
+		DestructionCallback = InDestructionCallback;
+	}
+
+private:
+	/** Hold a texture to be used for gpu texture transfers. */
+	TRefCountPtr<FRHITexture2D> Texture;
+
+	/** Called when the sample is destroyed by its pool. */
+	TFunction<void(TRefCountPtr<FRHITexture2D>)> DestructionCallback;
+
 public:
 	//~ IMediaPoolable interface
 
 	virtual void ShutdownPoolable() override
 	{
+		if (DestructionCallback)
+		{
+			DestructionCallback(Texture);
+		}
+
 		FreeSample();
 	}
 
@@ -219,8 +269,8 @@ protected:
 
 	/** Pointer to raw pixels */
 	TArray<uint8, TAlignedHeapAllocator<4096>> Buffer;
+	void* ExternalBuffer = nullptr;
 
 	/** Wheter the sample is in sRGB space and requires an explicit conversion to linear */
 	bool bIsSRGBInput;
 };
-
