Fix blackmagic media provider initializing with wrong framerate

diff --git a/Engine/Plugins/Media/BlackmagicMedia/Source/BlackmagicMedia/Private/Assets/BlackmagicTimecodeProvider.cpp b/Engine/Plugins/Media/BlackmagicMedia/Source/BlackmagicMedia/Private/Assets/BlackmagicTimecodeProvider.cpp
--- a/Engine/Plugins/Media/BlackmagicMedia/Source/BlackmagicMedia/Private/Assets/BlackmagicTimecodeProvider.cpp
+++ b/Engine/Plugins/Media/BlackmagicMedia/Source/BlackmagicMedia/Private/Assets/BlackmagicTimecodeProvider.cpp
@@ -184,7 +184,7 @@ bool UBlackmagicTimecodeProvider::Initialize(class UEngine* InEngine)
 	ChannelInfo.DeviceIndex = MediaConfiguration.MediaConnection.Device.DeviceIdentifier;
 
 	check(EventCallback == nullptr);
-	EventCallback = new BlackmagicTimecodeProviderHelpers::FEventCallback(ChannelInfo, GetFrameRate());
+	EventCallback = new BlackmagicTimecodeProviderHelpers::FEventCallback(ChannelInfo, MediaConfiguration.MediaMode.FrameRate);
 
 	BlackmagicDesign::FInputChannelOptions ChannelOptions;
 	ChannelOptions.CallbackPriority = 5;
