Fix lyraeditor build

diff --git a/Engine/Plugins/Media/BlackmagicMedia/Source/BlackmagicMediaEditor/Private/BlackmagicMediaEditorModule.cpp b/Engine/Plugins/Media/BlackmagicMedia/Source/BlackmagicMediaEditor/Private/BlackmagicMediaEditorModule.cpp
--- a/Engine/Plugins/Media/BlackmagicMedia/Source/BlackmagicMediaEditor/Private/BlackmagicMediaEditorModule.cpp
+++ b/Engine/Plugins/Media/BlackmagicMedia/Source/BlackmagicMediaEditor/Private/BlackmagicMediaEditorModule.cpp
@@ -23,17 +23,13 @@ public:
 	virtual void StartupModule() override
 	{
 		RegisterStyle();
-		FPropertyEditorModule& PropertyModule = FModuleManager::LoadModuleChecked<FPropertyEditorModule>("PropertyEditor");
-//		PropertyModule.RegisterCustomPropertyTypeLayout(FAjaMediaTimecodeConfiguration::StaticStruct()->GetFName(), FOnGetPropertyTypeCustomizationInstance::CreateStatic(&FAjaMediaTimecodeConfigurationCustomization::MakeInstance));
 	}
 
 	virtual void ShutdownModule() override
 	{
 		if (!UObjectInitialized() && !IsEngineExitRequested())
 		{
-			FPropertyEditorModule& PropertyModule = FModuleManager::LoadModuleChecked<FPropertyEditorModule>("PropertyEditor");
-			//PropertyModule.UnregisterCustomClassLayout(UAjaMediaSource::StaticClass()->GetFName());
-			UnregisterStyle();
+			UnregisterStyle();	
 		}
 	}
 
