Blackmagic genlock autodetect

diff --git a/Engine/Plugins/Media/BlackmagicMedia/Source/BlackmagicMedia/Private/Assets/BlackmagicCustomTimeStep.cpp b/Engine/Plugins/Media/BlackmagicMedia/Source/BlackmagicMedia/Private/Assets/BlackmagicCustomTimeStep.cpp
--- a/Engine/Plugins/Media/BlackmagicMedia/Source/BlackmagicMedia/Private/Assets/BlackmagicCustomTimeStep.cpp
+++ b/Engine/Plugins/Media/BlackmagicMedia/Source/BlackmagicMedia/Private/Assets/BlackmagicCustomTimeStep.cpp
@@ -21,8 +21,9 @@ namespace BlackmagicCustomTimeStepHelpers
 	class FInputEventCallback : public BlackmagicDesign::IInputEventCallback
 	{
 	public:
-		FInputEventCallback(const BlackmagicDesign::FChannelInfo& InChannelInfo, bool bInEnableOverrunDetection, uint32 InExpectedSyncDelta)
-			: RefCounter(0)
+		FInputEventCallback(UBlackmagicCustomTimeStep* InCustomTimeStep, const BlackmagicDesign::FChannelInfo& InChannelInfo, bool bInEnableOverrunDetection, uint32 InExpectedSyncDelta)
+			: WeakOwner(InCustomTimeStep)
+			, RefCounter(0)
 			, ChannelInfo(InChannelInfo)
 			, State(ECustomTimeStepSynchronizationState::Closed)
 			, WaitSyncEvent(nullptr)
@@ -176,8 +177,34 @@ namespace BlackmagicCustomTimeStepHelpers
 
 		virtual void OnFrameFormatChanged(const BlackmagicDesign::FFormatInfo& NewFormat) override
 		{
-			UE_LOG(LogBlackmagicMedia, Error, TEXT("The video format changed."));
-			State = ECustomTimeStepSynchronizationState::Error;
+			UBlackmagicCustomTimeStep* TimeStep = WeakOwner.Get();
+			if (TimeStep && TimeStep->bAutoDetectFormat)
+			{
+				TimeStep->MediaConfiguration.MediaMode.Resolution.X = NewFormat.Width;
+				TimeStep->MediaConfiguration.MediaMode.Resolution.Y = NewFormat.Height;
+				switch (NewFormat.FieldDominance)
+				{
+				case BlackmagicDesign::EFieldDominance::Interlaced:
+					TimeStep->MediaConfiguration.MediaMode.Standard = EMediaIOStandardType::Interlaced;
+					break;
+				case BlackmagicDesign::EFieldDominance::Progressive:
+					TimeStep->MediaConfiguration.MediaMode.Standard = EMediaIOStandardType::Progressive;
+					break;
+				case BlackmagicDesign::EFieldDominance::ProgressiveSegmentedFrame:
+					TimeStep->MediaConfiguration.MediaMode.Standard = EMediaIOStandardType::Progressive;
+					break;
+				default:
+					break;
+				}
+
+				TimeStep->MediaConfiguration.MediaMode.FrameRate = FFrameRate(NewFormat.FrameRateNumerator, NewFormat.FrameRateDenominator);
+			}
+			else
+			{
+				UE_LOG(LogBlackmagicMedia, Error, TEXT("The video format changed."));
+				State = ECustomTimeStepSynchronizationState::Error;
+			}
+
 			if (WaitSyncEvent)
 			{
 				WaitSyncEvent->Trigger();
@@ -193,6 +220,8 @@ namespace BlackmagicCustomTimeStepHelpers
 		}
 
 	private:
+		TWeakObjectPtr<UBlackmagicCustomTimeStep> WeakOwner;
+
 		TAtomic<int32> RefCounter;
 
 		BlackmagicDesign::FUniqueIdentifier BlackmagicIdendifier;
@@ -243,9 +272,10 @@ bool UBlackmagicCustomTimeStep::Initialize(class UEngine* InEngine)
 	BlackmagicDesign::FChannelInfo ChannelInfo;
 	ChannelInfo.DeviceIndex = MediaConfiguration.MediaConnection.Device.DeviceIdentifier;
 
-	InputEventCallback = new BlackmagicCustomTimeStepHelpers::FInputEventCallback(ChannelInfo, bEnableOverrunDetection, GetExpectedSyncCountDelta());
+	InputEventCallback = new BlackmagicCustomTimeStepHelpers::FInputEventCallback(this, ChannelInfo, bEnableOverrunDetection, GetExpectedSyncCountDelta());
 
 	BlackmagicDesign::FInputChannelOptions ChannelOptions;
+	ChannelOptions.bAutoDetect = bAutoDetectFormat;
 	ChannelOptions.CallbackPriority = 1;
 	ChannelOptions.FormatInfo.DisplayMode = MediaConfiguration.MediaMode.DeviceModeIdentifier;
 	ChannelOptions.FormatInfo.FrameRateNumerator = MediaConfiguration.MediaMode.FrameRate.Numerator;
diff --git a/Engine/Plugins/Media/BlackmagicMedia/Source/BlackmagicMedia/Public/BlackmagicCustomTimeStep.h b/Engine/Plugins/Media/BlackmagicMedia/Source/BlackmagicMedia/Public/BlackmagicCustomTimeStep.h
--- a/Engine/Plugins/Media/BlackmagicMedia/Source/BlackmagicMedia/Public/BlackmagicCustomTimeStep.h
+++ b/Engine/Plugins/Media/BlackmagicMedia/Source/BlackmagicMedia/Public/BlackmagicCustomTimeStep.h
@@ -3,7 +3,6 @@
 #pragma once
 
 #include "GenlockedCustomTimeStep.h"
-
 #include "MediaIOCoreDefinitions.h"
 
 #include "BlackmagicCustomTimeStep.generated.h"
@@ -35,6 +34,7 @@ public:
 	virtual bool IsLastSyncDataValid() const override;
 	virtual FFrameRate GetSyncRate() const override;
 	virtual bool WaitForSync() override;
+	virtual bool SupportsFormatAutoDetection() const { return true; };
 
 	//~ UObject interface
 	virtual void BeginDestroy() override;
