Fix 10 bit media capture for AJA/BM when using GPU Direct

diff --git a/Engine/Plugins/Media/AjaMedia/Source/AjaMediaOutput/Private/AjaMediaCapture.cpp b/Engine/Plugins/Media/AjaMedia/Source/AjaMediaOutput/Private/AjaMediaCapture.cpp
--- a/Engine/Plugins/Media/AjaMedia/Source/AjaMediaOutput/Private/AjaMediaCapture.cpp
+++ b/Engine/Plugins/Media/AjaMedia/Source/AjaMediaOutput/Private/AjaMediaCapture.cpp
@@ -394,12 +394,6 @@ bool UAjaMediaCapture::InitAJA(UAjaMediaOutput* InAjaMediaOutput)
 			UE_LOG(LogAjaMediaOutput, Warning, TEXT("GPU DMA is not supported with interlaced, defaulting to regular path."));
 			bGPUTextureTransferAvailable = false;
 		}
-
-		if (InAjaMediaOutput->PixelFormat != EAjaMediaOutputPixelFormat::PF_8BIT_YUV)
-		{
-			UE_LOG(LogAjaMediaOutput, Warning, TEXT("GPU DMA is not supported with pixel format 10 Bit YUV, defaulting to regular path."));
-			bGPUTextureTransferAvailable = false;
-		}
 	}
 
 	// Init Device options
diff --git a/Engine/Plugins/Media/BlackmagicMedia/Source/BlackmagicMediaOutput/Private/BlackmagicMediaCapture.cpp b/Engine/Plugins/Media/BlackmagicMedia/Source/BlackmagicMediaOutput/Private/BlackmagicMediaCapture.cpp
--- a/Engine/Plugins/Media/BlackmagicMedia/Source/BlackmagicMediaOutput/Private/BlackmagicMediaCapture.cpp
+++ b/Engine/Plugins/Media/BlackmagicMedia/Source/BlackmagicMediaOutput/Private/BlackmagicMediaCapture.cpp
@@ -497,12 +497,6 @@ bool UBlackmagicMediaCapture::InitBlackmagic(UBlackmagicMediaOutput* InBlackmagi
 			UE_LOG(LogBlackmagicMediaOutput, Warning, TEXT("GPU DMA is not supported with interlaced, defaulting to regular path."));
 			bGPUTextureTransferAvailable = false;
 		}
-
-		if (InBlackmagicMediaOutput->PixelFormat != EBlackmagicMediaOutputPixelFormat::PF_8BIT_YUV)
-		{
-			UE_LOG(LogBlackmagicMediaOutput, Warning, TEXT("GPU DMA is not supported with pixel format 10 Bit YUV, defaulting to regular path."));
-			bGPUTextureTransferAvailable = false;
-		}
 	}
 
 
