- Updating media capture implementation for new API

diff --git a/Engine/Plugins/Experimental/RemoteSession/Source/RemoteSession/Private/ImageProviders/RemoteSessionMediaOutput.cpp b/Engine/Plugins/Experimental/RemoteSession/Source/RemoteSession/Private/ImageProviders/RemoteSessionMediaOutput.cpp
--- a/Engine/Plugins/Experimental/RemoteSession/Source/RemoteSession/Private/ImageProviders/RemoteSessionMediaOutput.cpp
+++ b/Engine/Plugins/Experimental/RemoteSession/Source/RemoteSession/Private/ImageProviders/RemoteSessionMediaOutput.cpp
@@ -43,15 +43,7 @@ void URemoteSessionMediaCapture::OnFrameCaptured_RenderingThread(const FCaptureB
 	}
 }
 
-bool URemoteSessionMediaCapture::CaptureSceneViewportImpl(TSharedPtr<FSceneViewport>& InSceneViewport)
-{
-	CacheValues();
-
-	SetState(EMediaCaptureState::Capturing);
-	return true;
-}
-
-bool URemoteSessionMediaCapture::CaptureRenderTargetImpl(UTextureRenderTarget2D* InRenderTarget)
+bool URemoteSessionMediaCapture::InitializeCapture()
 {
 	CacheValues();
 
diff --git a/Engine/Plugins/Experimental/RemoteSession/Source/RemoteSession/Public/ImageProviders/RemoteSessionMediaOutput.h b/Engine/Plugins/Experimental/RemoteSession/Source/RemoteSession/Public/ImageProviders/RemoteSessionMediaOutput.h
--- a/Engine/Plugins/Experimental/RemoteSession/Source/RemoteSession/Public/ImageProviders/RemoteSessionMediaOutput.h
+++ b/Engine/Plugins/Experimental/RemoteSession/Source/RemoteSession/Public/ImageProviders/RemoteSessionMediaOutput.h
@@ -41,8 +41,7 @@ public:
 
 	//~ Begin UMediaCapture interface
 	virtual void OnFrameCaptured_RenderingThread(const FCaptureBaseData& InBaseData, TSharedPtr<FMediaCaptureUserData, ESPMode::ThreadSafe> InUserData, void* InBuffer, int32 Width, int32 Height, int BytesPerRow) override;
-	virtual bool CaptureSceneViewportImpl(TSharedPtr<FSceneViewport>& InSceneViewport) override;
-	virtual bool CaptureRenderTargetImpl(UTextureRenderTarget2D* InRenderTarget) override;
+	virtual bool InitializeCapture() override;
 	//~ End UMediaCapture interface
 
 private:
diff --git a/Engine/Plugins/Media/AjaMedia/Source/AjaMediaOutput/Private/AjaMediaCapture.cpp b/Engine/Plugins/Media/AjaMedia/Source/AjaMediaOutput/Private/AjaMediaCapture.cpp
--- a/Engine/Plugins/Media/AjaMedia/Source/AjaMediaOutput/Private/AjaMediaCapture.cpp
+++ b/Engine/Plugins/Media/AjaMedia/Source/AjaMediaOutput/Private/AjaMediaCapture.cpp
@@ -235,31 +235,23 @@ bool UAjaMediaCapture::ValidateMediaOutput() const
 	return true;
 }
 
-bool UAjaMediaCapture::CaptureSceneViewportImpl(TSharedPtr<FSceneViewport>& InSceneViewport)
+bool UAjaMediaCapture::InitializeCapture()
 {
 	UAjaMediaOutput* AjaMediaSource = CastChecked<UAjaMediaOutput>(MediaOutput);
 	const bool bResult = InitAJA(AjaMediaSource);
 	if (bResult)
 	{
-		ApplyViewportTextureAlpha(InSceneViewport);
 #if WITH_EDITOR
-		AjaMediaCaptureAnalytics::SendCaptureEvent(AjaMediaSource->GetRequestedSize(), FrameRate, TEXT("SceneViewport"));
+		AjaMediaCaptureAnalytics::SendCaptureEvent(AjaMediaSource->GetRequestedSize(), FrameRate, GetCaptureSourceType());
 #endif
 	}
 	return bResult;
 }
 
-bool UAjaMediaCapture::CaptureRenderTargetImpl(UTextureRenderTarget2D* InRenderTarget)
+bool UAjaMediaCapture::PostInitializeCaptureViewport(TSharedPtr<FSceneViewport>& InSceneViewport)
 {
-	UAjaMediaOutput* AjaMediaSource = CastChecked<UAjaMediaOutput>(MediaOutput);
-	bool bResult = InitAJA(AjaMediaSource);
-#if WITH_EDITOR
-	if (bResult)
-	{
-		AjaMediaCaptureAnalytics::SendCaptureEvent(AjaMediaSource->GetRequestedSize(), FrameRate, TEXT("RenderTarget"));
-	}
-#endif
-	return bResult;
+	ApplyViewportTextureAlpha(InSceneViewport);
+	return true;
 }
 
 bool UAjaMediaCapture::UpdateSceneViewportImpl(TSharedPtr<FSceneViewport>& InSceneViewport)
diff --git a/Engine/Plugins/Media/AjaMedia/Source/AjaMediaOutput/Public/AjaMediaCapture.h b/Engine/Plugins/Media/AjaMedia/Source/AjaMediaOutput/Public/AjaMediaCapture.h
--- a/Engine/Plugins/Media/AjaMedia/Source/AjaMediaOutput/Public/AjaMediaCapture.h
+++ b/Engine/Plugins/Media/AjaMedia/Source/AjaMediaOutput/Public/AjaMediaCapture.h
@@ -36,8 +36,8 @@ public:
 	virtual bool HasFinishedProcessing() const override;
 protected:
 	virtual bool ValidateMediaOutput() const override;
-	virtual bool CaptureSceneViewportImpl(TSharedPtr<FSceneViewport>& InSceneViewport) override;
-	virtual bool CaptureRenderTargetImpl(UTextureRenderTarget2D* InRenderTarget) override;
+	virtual bool InitializeCapture() override;
+	virtual bool PostInitializeCaptureViewport(TSharedPtr<FSceneViewport>& InSceneViewport) override;
 	virtual bool UpdateSceneViewportImpl(TSharedPtr<FSceneViewport>& InSceneViewport) override;
 	virtual bool UpdateRenderTargetImpl(UTextureRenderTarget2D* InRenderTarget) override;
 	virtual void StopCaptureImpl(bool bAllowPendingFrameToBeProcess) override;
diff --git a/Engine/Plugins/Media/BlackmagicMedia/Source/BlackmagicMediaOutput/Private/BlackmagicMediaCapture.cpp b/Engine/Plugins/Media/BlackmagicMedia/Source/BlackmagicMediaOutput/Private/BlackmagicMediaCapture.cpp
--- a/Engine/Plugins/Media/BlackmagicMedia/Source/BlackmagicMediaOutput/Private/BlackmagicMediaCapture.cpp
+++ b/Engine/Plugins/Media/BlackmagicMedia/Source/BlackmagicMediaOutput/Private/BlackmagicMediaCapture.cpp
@@ -347,30 +347,25 @@ bool UBlackmagicMediaCapture::ValidateMediaOutput() const
 	return true;
 }
 
-bool UBlackmagicMediaCapture::CaptureSceneViewportImpl(TSharedPtr<FSceneViewport>& InSceneViewport)
+bool UBlackmagicMediaCapture::InitializeCapture()
 {
 	UBlackmagicMediaOutput* BlackmagicMediaOutput = CastChecked<UBlackmagicMediaOutput>(MediaOutput);
-	bool bResult = InitBlackmagic(BlackmagicMediaOutput);
-	if (bResult)
+	BlackmagicMediaOutputPixelFormat = BlackmagicMediaOutput->PixelFormat;
+	bool bInitialized = InitBlackmagic(BlackmagicMediaOutput);
+	if(bInitialized)
 	{
-		ApplyViewportTextureAlpha(InSceneViewport);
-		BlackmagicMediaOutputPixelFormat = BlackmagicMediaOutput->PixelFormat;
 #if WITH_EDITOR
-		BlackmagicMediaCaptureAnalytics::SendCaptureEvent(BlackmagicMediaOutput->GetRequestedSize(), FrameRate, TEXT("SceneViewport"));
+		BlackmagicMediaCaptureAnalytics::SendCaptureEvent(BlackmagicMediaOutput->GetRequestedSize(), FrameRate, GetCaptureSourceType());
 #endif
 	}
-	return bResult;
+
+	return bInitialized;
 }
 
-bool UBlackmagicMediaCapture::CaptureRenderTargetImpl(UTextureRenderTarget2D* InRenderTarget)
+bool UBlackmagicMediaCapture::PostInitializeCaptureViewport(TSharedPtr<FSceneViewport>& InSceneViewport)
 {
-	UBlackmagicMediaOutput* BlackmagicMediaOutput = CastChecked<UBlackmagicMediaOutput>(MediaOutput);
-	bool bResult = InitBlackmagic(BlackmagicMediaOutput);
-	BlackmagicMediaOutputPixelFormat = BlackmagicMediaOutput->PixelFormat;
-#if WITH_EDITOR
-	BlackmagicMediaCaptureAnalytics::SendCaptureEvent(BlackmagicMediaOutput->GetRequestedSize(), FrameRate, TEXT("RenderTarget"));
-#endif
-	return bResult;
+	ApplyViewportTextureAlpha(InSceneViewport);
+	return true;
 }
 
 bool UBlackmagicMediaCapture::UpdateSceneViewportImpl(TSharedPtr<FSceneViewport>& InSceneViewport)
diff --git a/Engine/Plugins/Media/BlackmagicMedia/Source/BlackmagicMediaOutput/Public/BlackmagicMediaCapture.h b/Engine/Plugins/Media/BlackmagicMedia/Source/BlackmagicMediaOutput/Public/BlackmagicMediaCapture.h
--- a/Engine/Plugins/Media/BlackmagicMedia/Source/BlackmagicMediaOutput/Public/BlackmagicMediaCapture.h
+++ b/Engine/Plugins/Media/BlackmagicMedia/Source/BlackmagicMediaOutput/Public/BlackmagicMediaCapture.h
@@ -36,8 +36,8 @@ public:
 	virtual bool HasFinishedProcessing() const override;
 protected:
 	virtual bool ValidateMediaOutput() const override;
-	virtual bool CaptureSceneViewportImpl(TSharedPtr<FSceneViewport>& InSceneViewport) override;
-	virtual bool CaptureRenderTargetImpl(UTextureRenderTarget2D* InRenderTarget) override;
+	virtual bool InitializeCapture() override;
+	virtual bool PostInitializeCaptureViewport(TSharedPtr<FSceneViewport>& InSceneViewport) override;
 	virtual bool UpdateSceneViewportImpl(TSharedPtr<FSceneViewport>& InSceneViewport) override;
 	virtual bool UpdateRenderTargetImpl(UTextureRenderTarget2D* InRenderTarget) override;
 	virtual void StopCaptureImpl(bool bAllowPendingFrameToBeProcess) override;
diff --git a/Engine/Plugins/Media/MediaIOFramework/Source/MediaIOCore/Private/FileMediaCapture.cpp b/Engine/Plugins/Media/MediaIOFramework/Source/MediaIOCore/Private/FileMediaCapture.cpp
--- a/Engine/Plugins/Media/MediaIOFramework/Source/MediaIOCore/Private/FileMediaCapture.cpp
+++ b/Engine/Plugins/Media/MediaIOFramework/Source/MediaIOCore/Private/FileMediaCapture.cpp
@@ -136,32 +136,20 @@ void UFileMediaCapture::OnFrameCaptured_RenderingThread(const FCaptureBaseData&
 	}
 }
 
-bool UFileMediaCapture::CaptureSceneViewportImpl(TSharedPtr<FSceneViewport>& InSceneViewport)
+bool UFileMediaCapture::InitializeCapture()
 {
 	FModuleManager::Get().LoadModuleChecked<IImageWriteQueueModule>("ImageWriteQueue");
 	CacheMediaOutputValues();
 
 	SetState(EMediaCaptureState::Capturing);
-#if WITH_EDITOR
-	FileMediaCaptureAnalytics::SendCaptureEvent(ImageFormat, CompressionQuality, TEXT("SceneViewport"));
-#endif
-	return true;
-}
-
-
-bool UFileMediaCapture::CaptureRenderTargetImpl(UTextureRenderTarget2D* InRenderTarget)
-{
-	FModuleManager::Get().LoadModuleChecked<IImageWriteQueueModule>("ImageWriteQueue");
-	CacheMediaOutputValues();
 
-	SetState(EMediaCaptureState::Capturing);
 #if WITH_EDITOR
-	FileMediaCaptureAnalytics::SendCaptureEvent(ImageFormat, CompressionQuality, TEXT("RenderTarget"));
+	FileMediaCaptureAnalytics::SendCaptureEvent(ImageFormat, CompressionQuality, GetCaptureSourceType());
 #endif
+
 	return true;
 }
 
-
 void UFileMediaCapture::CacheMediaOutputValues()
 {
 	UFileMediaOutput* FileMediaOutput = CastChecked<UFileMediaOutput>(MediaOutput);
diff --git a/Engine/Plugins/Media/MediaIOFramework/Source/MediaIOCore/Public/FileMediaCapture.h b/Engine/Plugins/Media/MediaIOFramework/Source/MediaIOCore/Public/FileMediaCapture.h
--- a/Engine/Plugins/Media/MediaIOFramework/Source/MediaIOCore/Public/FileMediaCapture.h
+++ b/Engine/Plugins/Media/MediaIOFramework/Source/MediaIOCore/Public/FileMediaCapture.h
@@ -17,8 +17,7 @@ class MEDIAIOCORE_API UFileMediaCapture : public UMediaCapture
 
 protected:
 	virtual void OnFrameCaptured_RenderingThread(const FCaptureBaseData& InBaseData, TSharedPtr<FMediaCaptureUserData, ESPMode::ThreadSafe> InUserData, void* InBuffer, int32 Width, int32 Height, int32 BytesPerRow) override;
-	virtual bool CaptureSceneViewportImpl(TSharedPtr<FSceneViewport>& InSceneViewport) override;
-	virtual bool CaptureRenderTargetImpl(UTextureRenderTarget2D* InRenderTarget) override;
+	virtual bool InitializeCapture() override;
 
 private:
 	void CacheMediaOutputValues();
diff --git a/Engine/Plugins/Media/PixelStreamingVCam/Source/PixelStreamingVCam/Private/PixelStreamingMediaCapture.cpp b/Engine/Plugins/Media/PixelStreamingVCam/Source/PixelStreamingVCam/Private/PixelStreamingMediaCapture.cpp
--- a/Engine/Plugins/Media/PixelStreamingVCam/Source/PixelStreamingVCam/Private/PixelStreamingMediaCapture.cpp
+++ b/Engine/Plugins/Media/PixelStreamingVCam/Source/PixelStreamingVCam/Private/PixelStreamingMediaCapture.cpp
@@ -16,19 +16,16 @@ void UPixelStreamingMediaCapture::OnRHIResourceCaptured_RenderingThread(
 	}
 }
 
-bool UPixelStreamingMediaCapture::CaptureSceneViewportImpl(TSharedPtr<FSceneViewport>& InSceneViewport)
+bool UPixelStreamingMediaCapture::InitializeCapture()
 {
-	Viewport = InSceneViewport;
 	SetupVideoInput();
 	SetState(EMediaCaptureState::Capturing);
 	return true;
 }
 
-bool UPixelStreamingMediaCapture::CaptureRenderTargetImpl(UTextureRenderTarget2D* InRenderTarget)
+bool UPixelStreamingMediaCapture::PostInitializeCaptureViewport(TSharedPtr<FSceneViewport>& InSceneViewport)
 {
-	Viewport = nullptr;
-	SetupVideoInput();
-	SetState(EMediaCaptureState::Capturing);
+	Viewport = InSceneViewport;
 	return true;
 }
 
diff --git a/Engine/Plugins/Media/PixelStreamingVCam/Source/PixelStreamingVCam/Private/PixelStreamingMediaCapture.h b/Engine/Plugins/Media/PixelStreamingVCam/Source/PixelStreamingVCam/Private/PixelStreamingMediaCapture.h
--- a/Engine/Plugins/Media/PixelStreamingVCam/Source/PixelStreamingVCam/Private/PixelStreamingMediaCapture.h
+++ b/Engine/Plugins/Media/PixelStreamingVCam/Source/PixelStreamingVCam/Private/PixelStreamingMediaCapture.h
@@ -18,8 +18,8 @@ public:
 		TSharedPtr<FMediaCaptureUserData, ESPMode::ThreadSafe> InUserData,
 		FTextureRHIRef InTexture) override;
 
-	virtual bool CaptureSceneViewportImpl(TSharedPtr<FSceneViewport>& InSceneViewport) override;
-	virtual bool CaptureRenderTargetImpl(UTextureRenderTarget2D* InRenderTarget) override;
+	virtual bool InitializeCapture() override;
+	virtual bool PostInitializeCaptureViewport(TSharedPtr<FSceneViewport>& InSceneViewport) override;
 	virtual bool ShouldCaptureRHIResource() const { return true; }
 	virtual void StopCaptureImpl(bool bAllowPendingFrameToBeProcess) override {}
 	//~ End UMediaCapture interface
