Fix Localization warning in MediaIO

diff --git a/Engine/Plugins/Media/BlackmagicMedia/Source/BlackmagicMedia/Private/Blackmagic/BlackmagicDeviceProvider.cpp b/Engine/Plugins/Media/BlackmagicMedia/Source/BlackmagicMedia/Private/Blackmagic/BlackmagicDeviceProvider.cpp
--- a/Engine/Plugins/Media/BlackmagicMedia/Source/BlackmagicMedia/Private/Blackmagic/BlackmagicDeviceProvider.cpp
+++ b/Engine/Plugins/Media/BlackmagicMedia/Source/BlackmagicMedia/Private/Blackmagic/BlackmagicDeviceProvider.cpp
@@ -449,7 +449,7 @@ FText FBlackmagicDeviceProvider::ToText(const FMediaIOConfiguration& InConfigura
 {
 	if (bInIsAutoDetected)
 	{
-		return FText::Format(LOCTEXT("FMediaIOConfigurationToText", "[{0}] - {1} [device{2}/auto]")
+		return FText::Format(LOCTEXT("FMediaIOConfigurationToTextAutoDetect", "[{0}] - {1} [device{2}/auto]")
 		, InConfiguration.bIsInput ? LOCTEXT("In", "In") : LOCTEXT("Out", "Out")
 		, FText::FromName(InConfiguration.MediaConnection.Device.DeviceName)
 		, FText::AsNumber(InConfiguration.MediaConnection.Device.DeviceIdentifier)
diff --git a/Engine/Plugins/Media/MediaIOFramework/Source/MediaIOCore/Private/MediaIOCoreDeviceProvider.cpp b/Engine/Plugins/Media/MediaIOFramework/Source/MediaIOCore/Private/MediaIOCoreDeviceProvider.cpp
--- a/Engine/Plugins/Media/MediaIOFramework/Source/MediaIOCore/Private/MediaIOCoreDeviceProvider.cpp
+++ b/Engine/Plugins/Media/MediaIOFramework/Source/MediaIOCore/Private/MediaIOCoreDeviceProvider.cpp
@@ -33,7 +33,7 @@ FText IMediaIOCoreDeviceProvider::ToText(const FMediaIOConfiguration& InConfigur
 {
 	if (bInIsAutoDetected)
 	{
-		return FText::Format(LOCTEXT("FMediaIOConfigurationToText", "{0} - {1} [device{2}/auto]")
+		return FText::Format(LOCTEXT("FMediaIOConfigurationToTextAutoDetect", "{0} - {1} [device{2}/auto]")
 				, InConfiguration.bIsInput ? LOCTEXT("In", "In") : LOCTEXT("Out", "Out")
 				, FText::FromName(InConfiguration.MediaConnection.Device.DeviceName)
 				, FText::AsNumber(InConfiguration.MediaConnection.Device.DeviceIdentifier)
