Fix log spam when blackmagic signal is lost
Fix PSF switch on BM Input not working

diff --git a/Engine/Plugins/Media/BlackmagicMedia/Source/BlackmagicMedia/Private/Player/BlackmagicMediaPlayer.cpp b/Engine/Plugins/Media/BlackmagicMedia/Source/BlackmagicMedia/Private/Player/BlackmagicMediaPlayer.cpp
--- a/Engine/Plugins/Media/BlackmagicMedia/Source/BlackmagicMedia/Private/Player/BlackmagicMediaPlayer.cpp
+++ b/Engine/Plugins/Media/BlackmagicMedia/Source/BlackmagicMedia/Private/Player/BlackmagicMediaPlayer.cpp
@@ -268,7 +268,7 @@ namespace BlackmagicMediaPlayerHelpers
 
 				if (InFrameInfo.VideoBuffer)
 				{
-					const bool bIsProgressivePicture = InFrameInfo.FieldDominance != BlackmagicDesign::EFieldDominance::Interlaced;
+					const bool bIsProgressivePicture = InFrameInfo.FieldDominance == BlackmagicDesign::EFieldDominance::Progressive;
 					EMediaTextureSampleFormat SampleFormat = EMediaTextureSampleFormat::CharBGRA;
 					EMediaIOCoreEncodePixelFormat EncodePixelFormat = EMediaIOCoreEncodePixelFormat::CharUYVY;
 					FString OutputFilename = "";
