Fix cook by disabling blackmagic initialisation if we can't render

diff --git a/Engine/Plugins/Media/BlackmagicMedia/Source/BlackmagicCore/Private/BlackmagicCoreModule.cpp b/Engine/Plugins/Media/BlackmagicMedia/Source/BlackmagicCore/Private/BlackmagicCoreModule.cpp
--- a/Engine/Plugins/Media/BlackmagicMedia/Source/BlackmagicCore/Private/BlackmagicCoreModule.cpp
+++ b/Engine/Plugins/Media/BlackmagicMedia/Source/BlackmagicCore/Private/BlackmagicCoreModule.cpp
@@ -9,7 +9,11 @@ DEFINE_LOG_CATEGORY(LogBlackmagicCore);
 void FBlackmagicCoreModule::StartupModule()
 {
 	bCanForceBlackmagicUsage = FParse::Param(FCommandLine::Get(), TEXT("forceblackmagicusage"));
-	bInitialized = BlackmagicDesign::ApiInitialization();
+
+	if (CanUseBlackmagicCard())
+	{
+		bInitialized = BlackmagicDesign::ApiInitialization();
+	}
 }
 void FBlackmagicCoreModule::ShutdownModule()
 {
diff --git a/Engine/Plugins/Media/BlackmagicMedia/Source/BlackmagicMedia/Private/BlackmagicMediaModule.cpp b/Engine/Plugins/Media/BlackmagicMedia/Source/BlackmagicMedia/Private/BlackmagicMediaModule.cpp
--- a/Engine/Plugins/Media/BlackmagicMedia/Source/BlackmagicMedia/Private/BlackmagicMediaModule.cpp
+++ b/Engine/Plugins/Media/BlackmagicMedia/Source/BlackmagicMedia/Private/BlackmagicMediaModule.cpp
@@ -50,14 +50,21 @@ public:
 
 	virtual bool IsInitialized() const override
 	{
-		FBlackmagicCoreModule* CoreModule = FModuleManager::GetModulePtr<FBlackmagicCoreModule>("BlackmagicCore");
-		return CoreModule->IsInitialized();
+		if (FBlackmagicCoreModule* CoreModule = FModuleManager::GetModulePtr<FBlackmagicCoreModule>("BlackmagicCore"))
+		{
+			return CoreModule->IsInitialized();
+		}
+
+		return false;
 	}
 
 	virtual bool CanBeUsed() const override 
 	{
-		FBlackmagicCoreModule* CoreModule = FModuleManager::GetModulePtr<FBlackmagicCoreModule>("BlackmagicCore");
-		return CoreModule->CanUseBlackmagicCard();
+		if (FBlackmagicCoreModule* CoreModule = FModuleManager::GetModulePtr<FBlackmagicCoreModule>("BlackmagicCore"))
+		{
+			return CoreModule->CanUseBlackmagicCard();
+		}
+		return false;
 	}
 
 public:
