Fix clang build

diff --git a/Engine/Plugins/Media/BlackmagicMedia/Source/BlackmagicCore/Private/BlackmagicInputChannel.h b/Engine/Plugins/Media/BlackmagicMedia/Source/BlackmagicCore/Private/BlackmagicInputChannel.h
--- a/Engine/Plugins/Media/BlackmagicMedia/Source/BlackmagicCore/Private/BlackmagicInputChannel.h
+++ b/Engine/Plugins/Media/BlackmagicMedia/Source/BlackmagicCore/Private/BlackmagicInputChannel.h
@@ -14,6 +14,7 @@ namespace BlackmagicDesign
 		{
 		public:
 			FInputChannelNotificationCallback(FInputChannel* InOwner, IDeckLinkInput* InDeckLinkInput, IDeckLinkStatus* InDeckLinkStatus);
+			virtual ~FInputChannelNotificationCallback() = default;
 
 			HRESULT STDMETHODCALLTYPE QueryInterface(REFIID iid, LPVOID *ppv);
 
diff --git a/Engine/Plugins/Media/BlackmagicMedia/Source/BlackmagicCore/Private/BlackmagicOutputChannel.h b/Engine/Plugins/Media/BlackmagicMedia/Source/BlackmagicCore/Private/BlackmagicOutputChannel.h
--- a/Engine/Plugins/Media/BlackmagicMedia/Source/BlackmagicCore/Private/BlackmagicOutputChannel.h
+++ b/Engine/Plugins/Media/BlackmagicMedia/Source/BlackmagicCore/Private/BlackmagicOutputChannel.h
@@ -72,6 +72,7 @@ namespace BlackmagicDesign
 		{
 		public:
 			FOutputChannelNotificationCallback(FOutputChannel* InOwner, IDeckLinkOutput* InDeckLinkOutput);
+			virtual ~FOutputChannelNotificationCallback() = default;
 
 			HRESULT STDMETHODCALLTYPE QueryInterface(REFIID iid, LPVOID *ppv);
 
