Fix blackmagic gpudirect

diff --git a/Engine/Plugins/Media/BlackmagicMedia/Source/BlackmagicCore/Private/BlackmagicOutputChannel.cpp b/Engine/Plugins/Media/BlackmagicMedia/Source/BlackmagicCore/Private/BlackmagicOutputChannel.cpp
--- a/Engine/Plugins/Media/BlackmagicMedia/Source/BlackmagicCore/Private/BlackmagicOutputChannel.cpp
+++ b/Engine/Plugins/Media/BlackmagicMedia/Source/BlackmagicCore/Private/BlackmagicOutputChannel.cpp
@@ -270,6 +270,8 @@ namespace BlackmagicDesign
 				return false;
 			}
 
+			TextureTransfer = FGPUTextureTransferModule::Get().GetTextureTransfer();
+
 			if (ChannelOptions.bOutputAudio)
 			{
 				if (!InitializeAudio())
@@ -666,7 +668,7 @@ namespace BlackmagicDesign
 					if (OutputFrameToSchedule)
 					{
 #if PLATFORM_WINDOWS
-						if (ChannelOptions.bUseGPUDMA)
+						if (ChannelOptions.bUseGPUDMA && TextureTransfer)
 						{
 							void* DestinationBuffer;
 							OutputFrameToSchedule->BlackmagicFrame->GetBytes(&DestinationBuffer);
@@ -1338,7 +1340,7 @@ namespace BlackmagicDesign
 				for (FOutputFrame* Frame : AllFrames)
 				{
 					void* Buffer = nullptr;
-					if (SUCCEEDED(Frame->BlackmagicFrame->GetBytes(&Buffer)))
+					if (SUCCEEDED(Frame->BlackmagicFrame->GetBytes(&Buffer)) && TextureTransfer)
 					{
 						UE::GPUTextureTransfer::FRegisterDMABufferArgs Args;
 						Args.Height = VideoHeight;
@@ -1353,12 +1355,15 @@ namespace BlackmagicDesign
 				bRegisteredDMABuffers = true;
 			}
 
-			UE::GPUTextureTransfer::ETransferDirection Direction = UE::GPUTextureTransfer::ETransferDirection::GPU_TO_CPU;
-
-			TextureTransfer->TransferTexture(DestinationBuffer, InFrame.RHITexture, Direction);
-			AvailableWritingFrame->bVideoLineFilled = true;
-			AvailableWritingFrame->bVideoF2LineFilled = true;
-			AvailableWritingFrame->CopiedVideoBufferSize = VideoBufferSize;
+			if (TextureTransfer)
+			{
+				UE::GPUTextureTransfer::ETransferDirection Direction = UE::GPUTextureTransfer::ETransferDirection::GPU_TO_CPU;
+				
+				TextureTransfer->TransferTexture(DestinationBuffer, InFrame.RHITexture, Direction);
+				AvailableWritingFrame->bVideoLineFilled = true;
+				AvailableWritingFrame->bVideoF2LineFilled = true;
+				AvailableWritingFrame->CopiedVideoBufferSize = VideoBufferSize;
+			}
 			PushWhenFrameReady(AvailableWritingFrame);
 
 			return true;
@@ -1575,7 +1580,7 @@ namespace BlackmagicDesign
 					else
 					{
 	#if PLATFORM_WINDOWS
-						if (ChannelOptions.bUseGPUDMA)
+						if (ChannelOptions.bUseGPUDMA && TextureTransfer)
 						{
 
 							void* DestinationBuffer;
diff --git a/Engine/Plugins/Media/BlackmagicMedia/Source/BlackmagicCore/Private/BlackmagicOutputChannel.h b/Engine/Plugins/Media/BlackmagicMedia/Source/BlackmagicCore/Private/BlackmagicOutputChannel.h
--- a/Engine/Plugins/Media/BlackmagicMedia/Source/BlackmagicCore/Private/BlackmagicOutputChannel.h
+++ b/Engine/Plugins/Media/BlackmagicMedia/Source/BlackmagicCore/Private/BlackmagicOutputChannel.h
@@ -6,6 +6,8 @@
 #include "ThreadSafeObjectPool.h"
 #include <atomic>
 
+#include "GPUTextureTransferModule.h"
+
 class CustomAllocator : public IDeckLinkMemoryAllocator
 {
 public:
@@ -285,6 +287,8 @@ namespace BlackmagicDesign
 			Blackmagic::Private::ThreadSafeQueue<FAudioPacket*> AudioPacketsToSchedule;
 			Blackmagic::Private::ThreadSafeObjectPool<FAudioPacket> AudioPacketPool;
 
+			UE::GPUTextureTransfer::TextureTransferPtr TextureTransfer;
+
 			bool bIsSSE2Available;
 			bool bIsFieldAEven;
 			bool bIsFirstFetch;
diff --git a/Engine/Plugins/Media/BlackmagicMedia/Source/BlackmagicCore/Private/TextureDMA.cpp b/Engine/Plugins/Media/BlackmagicMedia/Source/BlackmagicCore/Private/TextureDMA.cpp
deleted file mode 100644
diff --git a/Engine/Plugins/Media/BlackmagicMedia/Source/BlackmagicCore/Public/BlackmagicLib.h b/Engine/Plugins/Media/BlackmagicMedia/Source/BlackmagicCore/Public/BlackmagicLib.h
--- a/Engine/Plugins/Media/BlackmagicMedia/Source/BlackmagicCore/Public/BlackmagicLib.h
+++ b/Engine/Plugins/Media/BlackmagicMedia/Source/BlackmagicCore/Public/BlackmagicLib.h
@@ -17,8 +17,6 @@ namespace UE::GPUTextureTransfer
 
 namespace BlackmagicDesign
 {
-	extern UE::GPUTextureTransfer::ITextureTransfer* TextureTransfer;
-
 	using LoggingCallbackPtr = void(*)(const TCHAR* Format, ...);
 	using FBlackmagicVideoFormat = int32_t; //BMDDisplayMode
 
@@ -347,32 +345,6 @@ namespace BlackmagicDesign
 		uint32_t FrameIdentifier = 0;
 	};
 
-	struct BLACKMAGICCORE_API FInitializeDMAArgs
-	{
-		ERHI RHI = ERHI::Invalid;
-		void* RHIDevice = nullptr;
-		void* RHICommandQueue = nullptr;
-
-		// Begin Vulkan Only
-		void* VulkanInstance = nullptr;
-		uint8_t RHIDeviceUUID[16] = { 0 };
-		// End Vulkan Only
-	};
-
-	struct BLACKMAGICCORE_API FRegisterDMABufferArgs
-	{
-		void* Buffer = nullptr;
-		uint32_t Stride = 0;
-		uint32_t Height = 0;
-		uint32_t Width = 0;
-	};
-
-	struct BLACKMAGICCORE_API FRegisterDMATextureArgs
-	{
-		FRHITexture* RHITexture = nullptr;
-		void* RHIResourceMemory = nullptr; // Vulkan only
-	};
-
 	/* BlackmagicDeviceScanner definition
 	*****************************************************************************/
 	class BLACKMAGICCORE_API BlackmagicDeviceScanner
@@ -474,13 +446,4 @@ namespace BlackmagicDesign
 	BLACKMAGICCORE_API bool SendVideoFrameData(const FChannelInfo& InChannelInfo, const FFrameDescriptor& InFrame);
 	BLACKMAGICCORE_API bool SendVideoFrameData(const FChannelInfo& InChannelInfo, FFrameDescriptor_GPUDMA& InFrame);
 	BLACKMAGICCORE_API bool SendAudioSamples(const FChannelInfo& InChannelInfo, const FAudioSamplesDescriptor& InSamples);
-
-	BLACKMAGICCORE_API bool InitializeDMA(const FInitializeDMAArgs& Args);
-	BLACKMAGICCORE_API void UninitializeDMA();
-	BLACKMAGICCORE_API bool RegisterDMATexture(const FRegisterDMATextureArgs& Args);
-	BLACKMAGICCORE_API bool UnregisterDMATexture(FRHITexture* InRHITexture);
-	BLACKMAGICCORE_API bool RegisterDMABuffer(const FRegisterDMABufferArgs& InArgs);
-	BLACKMAGICCORE_API bool UnregisterDMABuffer(void* InBuffer);
-	BLACKMAGICCORE_API bool LockDMATexture(FRHITexture* InRHITexture);
-	BLACKMAGICCORE_API bool UnlockDMATexture(FRHITexture* InRHITexture);
 };
diff --git a/Engine/Plugins/Media/BlackmagicMedia/Source/BlackmagicMediaOutput/BlackmagicMediaOutput.Build.cs b/Engine/Plugins/Media/BlackmagicMedia/Source/BlackmagicMediaOutput/BlackmagicMediaOutput.Build.cs
--- a/Engine/Plugins/Media/BlackmagicMedia/Source/BlackmagicMediaOutput/BlackmagicMediaOutput.Build.cs
+++ b/Engine/Plugins/Media/BlackmagicMedia/Source/BlackmagicMediaOutput/BlackmagicMediaOutput.Build.cs
@@ -28,6 +28,7 @@ namespace UnrealBuildTool.Rules
 					"Core",
 					"CoreUObject",
 					"Engine",
+					"GPUTextureTransfer",
 					"RenderCore",
 					"RHI",
 					"Slate"
diff --git a/Engine/Plugins/Media/BlackmagicMedia/Source/BlackmagicMediaOutput/Private/BlackmagicMediaCapture.cpp b/Engine/Plugins/Media/BlackmagicMedia/Source/BlackmagicMediaOutput/Private/BlackmagicMediaCapture.cpp
--- a/Engine/Plugins/Media/BlackmagicMedia/Source/BlackmagicMediaOutput/Private/BlackmagicMediaCapture.cpp
+++ b/Engine/Plugins/Media/BlackmagicMedia/Source/BlackmagicMediaOutput/Private/BlackmagicMediaCapture.cpp
@@ -5,6 +5,8 @@
 #include "BlackmagicLib.h"
 #include "BlackmagicMediaOutput.h"
 #include "BlackmagicMediaOutputModule.h"
+#include "GPUTextureTransfer.h"
+#include "GPUTextureTransferModule.h"
 #include "Engine/Engine.h"
 #include "HAL/Event.h"
 #include "HAL/IConsoleManager.h"
@@ -331,7 +333,11 @@ UBlackmagicMediaCapture::UBlackmagicMediaCapture()
 	, WakeUpEvent(nullptr)
 	, LastFrameDropCount_BlackmagicThread(0)
 {
-	bGPUTextureTransferAvailable = FBlackmagicMediaOutputModule::Get().IsGPUTextureTransferAvailable();
+	bGPUTextureTransferAvailable = FGPUTextureTransferModule::Get().IsAvailable();
+	if (bGPUTextureTransferAvailable)
+	{
+		TextureTransfer = FGPUTextureTransferModule::Get().GetTextureTransfer();
+	}
 }
 
 bool UBlackmagicMediaCapture::ValidateMediaOutput() const
@@ -392,7 +398,7 @@ void UBlackmagicMediaCapture::StopCaptureImpl(bool bAllowPendingFrameToBeProcess
 				{
 					for (FTextureRHIRef& Texture : TexturesToRelease)
 					{
-						BlackmagicDesign::UnregisterDMATexture(Texture->GetTexture2D());
+						TextureTransfer->UnregisterTexture(Texture->GetTexture2D());
 					}
 
 					TexturesToRelease.Reset();
@@ -571,30 +577,54 @@ bool UBlackmagicMediaCapture::InitBlackmagic(UBlackmagicMediaOutput* InBlackmagi
 
 void UBlackmagicMediaCapture::LockDMATexture_RenderThread(FTextureRHIRef InTexture)
 {
-	if (ShouldCaptureRHIResource())
+	if (ShouldCaptureRHIResource() && TextureTransfer)
 	{
 		if (!TexturesToRelease.Contains(InTexture))
 		{
 			TexturesToRelease.Add(InTexture);
 
-			BlackmagicDesign::FRegisterDMATextureArgs Args;
-			Args.RHITexture = InTexture->GetTexture2D();
+			FRHITexture2D* Texture = InTexture->GetTexture2D();
+			UE::GPUTextureTransfer::FRegisterDMATextureArgs Args;
+			Args.RHITexture = Texture;
+
+			Args.RHIResourceMemory = nullptr; // = Texture->GetNativeResource(); todo: VulkanTexture->Surface->GetAllocationHandle for Vulkan
+			Args.Width = Texture->GetDesc().GetSize().X;
+			Args.Height = Texture->GetDesc().GetSize().Y;
+
+			if (Args.RHITexture->GetFormat() == EPixelFormat::PF_B8G8R8A8)
+			{
+				Args.PixelFormat = UE::GPUTextureTransfer::EPixelFormat::PF_8Bit;
+				Args.Stride = Args.Width * 4;
+			}
+			else if (Args.RHITexture->GetFormat() == EPixelFormat::PF_R32G32B32A32_UINT)
+			{
+				Args.PixelFormat = UE::GPUTextureTransfer::EPixelFormat::PF_10Bit;
+				Args.Stride = Args.Width * 16;
+			}
+			else
+			{
+				checkf(false, TEXT("Format not supported"));
+			}
 			//Args.RHIResourceMemory = Texture->GetNativeResource(); todo: VulkanTexture->Surface->GetAllocationHandle for Vulkan
-			BlackmagicDesign::RegisterDMATexture(Args);
+
+			TextureTransfer->RegisterTexture(Args);
 
 		}
 
 		{
 			TRACE_CPUPROFILER_EVENT_SCOPE(UBlackmagicMediaCapture::LockDMATexture);
-			BlackmagicDesign::LockDMATexture(InTexture->GetTexture2D());
+			TextureTransfer->LockTexture(InTexture->GetTexture2D());
 		}
 	}
 }
 
 void UBlackmagicMediaCapture::UnlockDMATexture_RenderThread(FTextureRHIRef InTexture)
 {
-	TRACE_CPUPROFILER_EVENT_SCOPE(UBlackmagicMediaCapture::UnlockDMATexture);
-	BlackmagicDesign::UnlockDMATexture(InTexture->GetTexture2D());
+	if (TextureTransfer)
+	{
+		TRACE_CPUPROFILER_EVENT_SCOPE(UBlackmagicMediaCapture::UnlockDMATexture);
+		TextureTransfer->UnlockTexture(InTexture->GetTexture2D());
+	}
 }
 
 void UBlackmagicMediaCapture::OnFrameCaptured_RenderingThread(const FCaptureBaseData& InBaseData, TSharedPtr<FMediaCaptureUserData, ESPMode::ThreadSafe> InUserData, void* InBuffer, int32 Width, int32 Height, int32 BytesPerRow)
diff --git a/Engine/Plugins/Media/BlackmagicMedia/Source/BlackmagicMediaOutput/Private/BlackmagicMediaOutputModule.cpp b/Engine/Plugins/Media/BlackmagicMedia/Source/BlackmagicMediaOutput/Private/BlackmagicMediaOutputModule.cpp
--- a/Engine/Plugins/Media/BlackmagicMedia/Source/BlackmagicMediaOutput/Private/BlackmagicMediaOutputModule.cpp
+++ b/Engine/Plugins/Media/BlackmagicMedia/Source/BlackmagicMediaOutput/Private/BlackmagicMediaOutputModule.cpp
@@ -1,98 +1,8 @@
 // Copyright Epic Games, Inc. All Rights Reserved.
 
 #include "BlackmagicMediaOutputModule.h"
-
-#include "BlackmagicLib.h"
-#include "GenericPlatform/GenericPlatformDriver.h"
-#include "RenderingThread.h"
-#include "RHI.h"
+#include "Modules/ModuleManager.h"
 
 DEFINE_LOG_CATEGORY(LogBlackmagicMediaOutput);
 
-FBlackmagicMediaOutputModule& FBlackmagicMediaOutputModule::Get()
-{
-	return FModuleManager::LoadModuleChecked<FBlackmagicMediaOutputModule>(TEXT("BlackmagicMediaOutput"));
-}
-
-void FBlackmagicMediaOutputModule::StartupModule()
-{
-	const auto DMAInitializationFunc = [this]()
-	{
-		static const TArray<FString> SupportedGPUPrefixes = {
-			TEXT("RTX A4"),
-			TEXT("RTX A5"),
-			TEXT("RTX A6"),
-			TEXT("Quadro")
-		};
-
-		const FGPUDriverInfo GPUDriverInfo = FPlatformMisc::GetGPUDriverInfo(GRHIAdapterName);
-		bIsGPUTextureTransferAvailable = GPUDriverInfo.IsNVIDIA() && !FModuleManager::Get().IsModuleLoaded("RenderDocPlugin");
-		bIsGPUTextureTransferAvailable &= !GPUDriverInfo.DeviceDescription.Contains(TEXT("Tesla"));
-
-		if (bIsGPUTextureTransferAvailable)
-		{
-			bIsGPUTextureTransferAvailable = false;
-			for (const FString& GPUPrefix : SupportedGPUPrefixes)
-			{
-				if (GPUDriverInfo.DeviceDescription.Contains(GPUPrefix))
-				{
-					bIsGPUTextureTransferAvailable = true;
-					break;
-				}
-			}
-		}
-
-
-		if (bIsGPUTextureTransferAvailable)
-		{
-			if (!GDynamicRHI)
-			{
-				return;
-			}
-
-			auto GetRHI = []()
-			{
-				switch (RHIGetInterfaceType())
-				{
-				case ERHIInterfaceType::D3D12: return BlackmagicDesign::ERHI::D3D12;
-				case ERHIInterfaceType::Vulkan: return BlackmagicDesign::ERHI::Vulkan;
-				default: return BlackmagicDesign::ERHI::Invalid;
-				}
-			};
-
-			BlackmagicDesign::FInitializeDMAArgs Args;
-			BlackmagicDesign::ERHI RHI = GetRHI();
-			Args.RHI = RHI;
-			Args.RHIDevice = GDynamicRHI->RHIGetNativeDevice();
-			Args.RHICommandQueue = GDynamicRHI->RHIGetNativeGraphicsQueue();
-
-			bIsGPUTextureTransferAvailable = BlackmagicDesign::InitializeDMA(Args);
-		}
-	};
-
-
-	auto DMAUninitialize = [this]()
-	{
-		if (bIsGPUTextureTransferAvailable)
-		{
-			BlackmagicDesign::UninitializeDMA();
-		}
-	};
-
-	//Postpone initialization after all modules have been loaded to be sure Blackmagic library has been loaded
-	FCoreDelegates::OnAllModuleLoadingPhasesComplete.AddLambda(DMAInitializationFunc);
-	//Same for shutdown, uninitialize ourselves before library is unloaded
-	FCoreDelegates::OnEnginePreExit.AddLambda(DMAUninitialize);
-}
-
-void FBlackmagicMediaOutputModule::ShutdownModule()
-{
-	
-}
-
-bool FBlackmagicMediaOutputModule::IsGPUTextureTransferAvailable() const
-{
-	return bIsGPUTextureTransferAvailable;
-}
-
-IMPLEMENT_MODULE(FBlackmagicMediaOutputModule, BlackmagicMediaOutput)
+IMPLEMENT_MODULE(FDefaultModuleImpl, BlackmagicMediaOutput)
diff --git a/Engine/Plugins/Media/BlackmagicMedia/Source/BlackmagicMediaOutput/Private/BlackmagicMediaOutputModule.h b/Engine/Plugins/Media/BlackmagicMedia/Source/BlackmagicMediaOutput/Private/BlackmagicMediaOutputModule.h
--- a/Engine/Plugins/Media/BlackmagicMedia/Source/BlackmagicMediaOutput/Private/BlackmagicMediaOutputModule.h
+++ b/Engine/Plugins/Media/BlackmagicMedia/Source/BlackmagicMediaOutput/Private/BlackmagicMediaOutputModule.h
@@ -2,23 +2,6 @@
 
 #pragma once
 
-#include "Modules/ModuleManager.h"
-
-#include "CoreMinimal.h"
+#include "Logging/LogMacros.h"
 
 DECLARE_LOG_CATEGORY_EXTERN(LogBlackmagicMediaOutput, Log, All);
-
-
-class FBlackmagicMediaOutputModule : public IModuleInterface
-{
-public:
-	static FBlackmagicMediaOutputModule& Get();
-	virtual void StartupModule() override;
-	virtual void ShutdownModule() override;
-
-
-	bool IsGPUTextureTransferAvailable() const;
-
-private:
-	bool bIsGPUTextureTransferAvailable = false;
-};
\ No newline at end of file
diff --git a/Engine/Plugins/Media/BlackmagicMedia/Source/BlackmagicMediaOutput/Public/BlackmagicMediaCapture.h b/Engine/Plugins/Media/BlackmagicMedia/Source/BlackmagicMediaOutput/Public/BlackmagicMediaCapture.h
--- a/Engine/Plugins/Media/BlackmagicMedia/Source/BlackmagicMediaOutput/Public/BlackmagicMediaCapture.h
+++ b/Engine/Plugins/Media/BlackmagicMedia/Source/BlackmagicMediaOutput/Public/BlackmagicMediaCapture.h
@@ -22,6 +22,11 @@ namespace BlackmagicDesign
 	struct FTimecode;
 }
 
+namespace UE::GPUTextureTransfer
+{
+	class ITextureTransfer;
+}
+
 /**
  * Output Media for Blackmagic streams.
  * The output format could be any of EBlackmagicMediaOutputPixelFormat.
@@ -108,4 +113,6 @@ private:
 
 	/** The last time OutputAudio was called. Used to adjust the number of audio samples to grab on each frame depending on the frame time. */
 	double OutputAudioTimestamp;
+
+	TSharedPtr<UE::GPUTextureTransfer::ITextureTransfer> TextureTransfer;
 };
