Fix deprecation warning in blackmagic

diff --git a/Engine/Plugins/Media/BlackmagicMedia/Source/BlackmagicCore/Private/Platform/Linux/LinuxPlatform.cpp b/Engine/Plugins/Media/BlackmagicMedia/Source/BlackmagicCore/Private/Platform/Linux/LinuxPlatform.cpp
--- a/Engine/Plugins/Media/BlackmagicMedia/Source/BlackmagicCore/Private/Platform/Linux/LinuxPlatform.cpp
+++ b/Engine/Plugins/Media/BlackmagicMedia/Source/BlackmagicCore/Private/Platform/Linux/LinuxPlatform.cpp
@@ -5,6 +5,9 @@
 
 #if PLATFORM_LINUX
 
+#include "Containers/StringConv.h"
+#include "Misc/CString.h"
+
 #include "Common.h"
 #include <pthread.h>
 #include <sys/mman.h>
@@ -94,10 +97,11 @@ namespace BlackmagicPlatform
 		HRESULT Result = Device->GetDisplayName(&DisplayName);
 		if (Result == S_OK)
 		{
-			std::wstring_convert<std::codecvt_utf8_utf16<TCHAR>, TCHAR> Convert;
-			auto OutSize = Convert.from_bytes(DisplayName).copy(OutDisplayName, Size - 1, 0);
-			OutDisplayName[OutSize] = 0;
-			free((void *)DisplayName);
+			TStringConvert<UTF8CHAR, TCHAR> StringConvert;
+			int32 SourceLength = FCStringAnsi::Strlen(DisplayName) + 1;
+			int32 ConvertedLength = StringConvert.ConvertedLength(DisplayName, SourceLength);
+			check(ConvertedLength < Size);
+			StringConvert.Convert(OutDisplayName, ConvertedLength, DisplayName, SourceLength);
 			return true;
 		}
 
