Add anythread path to blackmagic

diff --git a/Engine/Plugins/Media/BlackmagicMedia/Source/BlackmagicMediaOutput/Private/BlackmagicMediaCapture.cpp b/Engine/Plugins/Media/BlackmagicMedia/Source/BlackmagicMediaOutput/Private/BlackmagicMediaCapture.cpp
--- a/Engine/Plugins/Media/BlackmagicMedia/Source/BlackmagicMediaOutput/Private/BlackmagicMediaCapture.cpp
+++ b/Engine/Plugins/Media/BlackmagicMedia/Source/BlackmagicMediaOutput/Private/BlackmagicMediaCapture.cpp
@@ -392,7 +392,7 @@
 	{
 		{
 			// Prevent the rendering thread from copying while we are stopping the capture.
-			FScopeLock ScopeLock(&RenderThreadCriticalSection);
+			FScopeLock ScopeLock(&CopyingCriticalSection);
 			ENQUEUE_RENDER_COMMAND(BlackmagicMediaCaptureInitialize)(
 				[this](FRHICommandListImmediate& RHICmdList) mutable
 				{
@@ -468,6 +468,81 @@
 	}
 }
 
+void UBlackmagicMediaCapture::OnFrameCapturedInternal_AnyThread(const FCaptureBaseData& InBaseData, TSharedPtr<FMediaCaptureUserData, ESPMode::ThreadSafe> InUserData, const FMediaCaptureResourceData& InResourceData)
+{
+	// Prevent this thread from copying while we are stopping the capture.
+	FScopeLock ScopeLock(&CopyingCriticalSection);
+	
+	if (EventCallback)
+	{
+		BlackmagicDesign::FTimecode Timecode = BlackmagicMediaCaptureDevice::ConvertToBlackmagicTimecode(InBaseData.SourceFrameTimecode, InBaseData.SourceFrameTimecodeFramerate.AsDecimal(), FrameRate.AsDecimal());
+
+		if (bEncodeTimecodeInTexel)
+		{
+			BlackmagicMediaCaptureHelpers::EncodeTimecodeInTexel(BlackmagicMediaOutputPixelFormat, GetConversionOperation(), Timecode, InResourceData.Buffer, InResourceData.Width, InResourceData.Height);
+		}
+
+		BlackmagicDesign::FFrameDescriptor Frame;
+		Frame.VideoBuffer = reinterpret_cast<uint8_t*>(InResourceData.Buffer);
+		Frame.VideoWidth = InResourceData.Width;
+		Frame.VideoHeight = InResourceData.Height;
+		Frame.Timecode = Timecode;
+		Frame.FrameIdentifier = InBaseData.SourceFrameNumber;
+		Frame.bEvenFrame = GFrameCounterRenderThread % 2 == 0;
+
+		bool bSent = false;
+		{
+			TRACE_CPUPROFILER_EVENT_SCOPE(UBlackmagicMediaCapture::SendVideoFrameData);
+			bSent = EventCallback->SendVideoFrameData(Frame);
+		}
+		
+		
+		if (bLogDropFrame && !bSent)
+		{
+			UE_LOG(LogBlackmagicMediaOutput, Warning, TEXT("Frame couldn't be sent to Blackmagic device. Engine might be running faster than output."));
+		}
+
+		OutputAudio_AnyThread(InBaseData, Timecode);
+
+		if (bBlackmagicWritInputRawDataCmdEnable)
+		{
+			FString OutputFilename;
+			uint32 Stride = 0;
+
+			switch (BlackmagicMediaOutputPixelFormat)
+			{
+			case EBlackmagicMediaOutputPixelFormat::PF_8BIT_YUV:
+				if (GetConversionOperation() == EMediaCaptureConversionOperation::RGBA8_TO_YUV_8BIT)
+				{
+					OutputFilename = TEXT("Blackmagic_Input_8_YUV");
+					Stride = InResourceData.Width * 4;
+					break;
+				}
+				else
+				{
+					OutputFilename = TEXT("Blackmagic_Input_8_RGBA");
+					Stride = InResourceData.Width * 4;
+					break;
+				}
+			case EBlackmagicMediaOutputPixelFormat::PF_10BIT_YUV:
+				OutputFilename = TEXT("Blackmagic_Input_10_YUV");
+				Stride = InResourceData.Width * 16;
+				break;
+			}
+
+			MediaIOCoreFileWriter::WriteRawFile(OutputFilename, reinterpret_cast<uint8*>(InResourceData.Buffer), Stride * InResourceData.Height);
+			bBlackmagicWritInputRawDataCmdEnable = false;
+		}
+
+
+		WaitForSync_AnyThread();
+	}
+	else if (GetState() != EMediaCaptureState::Stopped)
+	{
+		SetState(EMediaCaptureState::Error);
+	}
+}
+
 bool UBlackmagicMediaCapture::HasFinishedProcessing() const
 {
 	return Super::HasFinishedProcessing() || EventCallback == nullptr;
@@ -629,90 +704,39 @@
 
 void UBlackmagicMediaCapture::OnFrameCaptured_RenderingThread(const FCaptureBaseData& InBaseData, TSharedPtr<FMediaCaptureUserData, ESPMode::ThreadSafe> InUserData, void* InBuffer, int32 Width, int32 Height, int32 BytesPerRow)
 {
-	// Prevent the rendering thread from copying while we are stopping the capture.
 	TRACE_CPUPROFILER_EVENT_SCOPE(UBlackmagicMediaCapture::OnFrameCaptured_RenderingThread);
-	FScopeLock ScopeLock(&RenderThreadCriticalSection);
-	if (EventCallback)
-	{
-		BlackmagicDesign::FTimecode Timecode = BlackmagicMediaCaptureDevice::ConvertToBlackmagicTimecode(InBaseData.SourceFrameTimecode, InBaseData.SourceFrameTimecodeFramerate.AsDecimal(), FrameRate.AsDecimal());
-
-		if (bEncodeTimecodeInTexel)
-		{
-			BlackmagicMediaCaptureHelpers::EncodeTimecodeInTexel(BlackmagicMediaOutputPixelFormat, GetConversionOperation(), Timecode, InBuffer, Width, Height);
-		}
-
-		BlackmagicDesign::FFrameDescriptor Frame;
-		Frame.VideoBuffer = reinterpret_cast<uint8_t*>(InBuffer);
-		Frame.VideoWidth = Width;
-		Frame.VideoHeight = Height;
-		Frame.Timecode = Timecode;
-		Frame.FrameIdentifier = InBaseData.SourceFrameNumber;
-		Frame.bEvenFrame = GFrameCounterRenderThread % 2 == 0;
-
-		bool bSent = false;
-		{
-			TRACE_CPUPROFILER_EVENT_SCOPE(UBlackmagicMediaCapture::SendVideoFrameData);
-			bSent = EventCallback->SendVideoFrameData(Frame);
-		}
-		
-		
-		if (bLogDropFrame && !bSent)
-		{
-			UE_LOG(LogBlackmagicMediaOutput, Warning, TEXT("Frame couldn't be sent to Blackmagic device. Engine might be running faster than output."));
-		}
-
-		OutputAudio_RenderingThread(InBaseData, Timecode);
-
-		if (bBlackmagicWritInputRawDataCmdEnable)
-		{
-			FString OutputFilename;
-			uint32 Stride = 0;
-
-			switch (BlackmagicMediaOutputPixelFormat)
-			{
-			case EBlackmagicMediaOutputPixelFormat::PF_8BIT_YUV:
-				if (GetConversionOperation() == EMediaCaptureConversionOperation::RGBA8_TO_YUV_8BIT)
-				{
-					OutputFilename = TEXT("Blackmagic_Input_8_YUV");
-					Stride = Width * 4;
-					break;
-				}
-				else
-				{
-					OutputFilename = TEXT("Blackmagic_Input_8_RGBA");
-					Stride = Width * 4;
-					break;
-				}
-			case EBlackmagicMediaOutputPixelFormat::PF_10BIT_YUV:
-				OutputFilename = TEXT("Blackmagic_Input_10_YUV");
-				Stride = Width * 16;
-				break;
-			}
-
-			MediaIOCoreFileWriter::WriteRawFile(OutputFilename, reinterpret_cast<uint8*>(InBuffer), Stride * Height);
-			bBlackmagicWritInputRawDataCmdEnable = false;
-		}
-
-
-		WaitForSync_RenderingThread();
-	}
-	else if (GetState() != EMediaCaptureState::Stopped)
-	{
-		SetState(EMediaCaptureState::Error);
-	}
+	FMediaCaptureResourceData ResourceData;
+	ResourceData.Buffer = InBuffer;
+	ResourceData.Width = Width;
+	ResourceData.Height = Height;
+	ResourceData.BytesPerRow = BytesPerRow;
+	
+	OnFrameCaptured_AnyThread(InBaseData, InUserData, ResourceData);
 }
 
+void UBlackmagicMediaCapture::OnFrameCaptured_AnyThread(const FCaptureBaseData& InBaseData, TSharedPtr<FMediaCaptureUserData, ESPMode::ThreadSafe> InUserData, const FMediaCaptureResourceData& InResourceData)
+{
+	TRACE_CPUPROFILER_EVENT_SCOPE(UBlackmagicMediaCapture::OnFrameCaptured_AnyThread);
+
+	OnFrameCapturedInternal_AnyThread(InBaseData, InUserData, InResourceData);
+}
+
 void UBlackmagicMediaCapture::OnRHIResourceCaptured_RenderingThread(const FCaptureBaseData& InBaseData, TSharedPtr<FMediaCaptureUserData, ESPMode::ThreadSafe> InUserData, FTextureRHIRef InTexture)
 {
+	TRACE_CPUPROFILER_EVENT_SCOPE(UBlackmagicMediaCapture::OnRHIResourceCaptured_RenderingThread);
+	OnRHIResourceCaptured_AnyThread(InBaseData, InUserData, InTexture);
+}
+
+void UBlackmagicMediaCapture::OnRHIResourceCaptured_AnyThread(const FCaptureBaseData& InBaseData, TSharedPtr<FMediaCaptureUserData, ESPMode::ThreadSafe> InUserData, FTextureRHIRef InTexture)
+{
 	if (!InTexture)
 	{
 		return;
 	}
+	
 	// Prevent the rendering thread from copying while we are stopping the capture.
-	TRACE_CPUPROFILER_EVENT_SCOPE(UBlackmagicMediaCapture::OnFrameCaptured_RenderingThread);
-
-	FScopeLock ScopeLock(&RenderThreadCriticalSection);
-
+	TRACE_CPUPROFILER_EVENT_SCOPE(UBlackmagicMediaCapture::OnRHIResourceCaptured_AnyThread);
+	FScopeLock ScopeLock(&CopyingCriticalSection);
 
 	if (EventCallback)
 	{
@@ -735,9 +759,9 @@
 			UE_LOG(LogBlackmagicMediaOutput, Warning, TEXT("Frame couldn't be sent to Blackmagic device. Engine might be running faster than output."));
 		}
 
-		OutputAudio_RenderingThread(InBaseData, Timecode);
+		OutputAudio_AnyThread(InBaseData, Timecode);
 
-		WaitForSync_RenderingThread();
+		WaitForSync_AnyThread();
 	}
 	else if (GetState() != EMediaCaptureState::Stopped)
 	{
@@ -745,11 +769,11 @@
 	}
 }
 
-void UBlackmagicMediaCapture::WaitForSync_RenderingThread()
+void UBlackmagicMediaCapture::WaitForSync_AnyThread()
 {
 	if (bWaitForSyncEvent)
 	{
-		if (WakeUpEvent && GetState() == EMediaCaptureState::Capturing) // In render thread, could be shutdown in a middle of a frame
+		if (WakeUpEvent && GetState() == EMediaCaptureState::Capturing) // Could be shutdown in a middle of a frame 
 		{
 			const uint32 NumberOfMilliseconds = 1000;
 			if (!WakeUpEvent->Wait(NumberOfMilliseconds))
@@ -761,7 +785,7 @@
 	}
 }
 
-void UBlackmagicMediaCapture::OutputAudio_RenderingThread(const FCaptureBaseData& InBaseData, const BlackmagicDesign::FTimecode& InTimecode)
+void UBlackmagicMediaCapture::OutputAudio_AnyThread(const FCaptureBaseData& InBaseData, const BlackmagicDesign::FTimecode& InTimecode)
 {
 	if (bOutputAudio && AudioOutput)
 	{
diff --git a/Engine/Plugins/Media/BlackmagicMedia/Source/BlackmagicMediaOutput/Public/BlackmagicMediaCapture.h b/Engine/Plugins/Media/BlackmagicMedia/Source/BlackmagicMediaOutput/Public/BlackmagicMediaCapture.h
--- a/Engine/Plugins/Media/BlackmagicMedia/Source/BlackmagicMediaOutput/Public/BlackmagicMediaCapture.h
+++ b/Engine/Plugins/Media/BlackmagicMedia/Source/BlackmagicMediaOutput/Public/BlackmagicMediaCapture.h
@@ -51,6 +51,8 @@ protected:
 	virtual bool ShouldCaptureRHIResource() const override;
 
 	virtual void OnFrameCaptured_RenderingThread(const FCaptureBaseData& InBaseData, TSharedPtr<FMediaCaptureUserData, ESPMode::ThreadSafe> InUserData, void* InBuffer, int32 Width, int32 Height, int32 BytesPerRow) override;
+	virtual void OnFrameCaptured_AnyThread(const FCaptureBaseData& InBaseData, TSharedPtr<FMediaCaptureUserData, ESPMode::ThreadSafe> InUserData, const FMediaCaptureResourceData& InResourceData) override;
+	virtual void OnRHIResourceCaptured_AnyThread(const FCaptureBaseData& InBaseData, TSharedPtr<FMediaCaptureUserData, ESPMode::ThreadSafe> InUserData, FTextureRHIRef InTexture) override;
 	virtual void OnRHIResourceCaptured_RenderingThread(const FCaptureBaseData& InBaseData, TSharedPtr<FMediaCaptureUserData, ESPMode::ThreadSafe> InUserData, FTextureRHIRef InTexture) override;
 	virtual void LockDMATexture_RenderThread(FTextureRHIRef InTexture) override;
 	virtual void UnlockDMATexture_RenderThread(FTextureRHIRef InTexture) override;
@@ -61,10 +63,11 @@ private:
 
 private:
 	bool InitBlackmagic(UBlackmagicMediaOutput* InMediaOutput);
-	void WaitForSync_RenderingThread();
-	void OutputAudio_RenderingThread(const FCaptureBaseData& InBaseData, const BlackmagicDesign::FTimecode& Timecode);
+	void WaitForSync_AnyThread();
+	void OutputAudio_AnyThread(const FCaptureBaseData& InBaseData, const BlackmagicDesign::FTimecode& Timecode);
 	void ApplyViewportTextureAlpha(TSharedPtr<FSceneViewport> InSceneViewport);
 	void RestoreViewportTextureAlpha(TSharedPtr<FSceneViewport> InSceneViewport);
+	void OnFrameCapturedInternal_AnyThread(const FCaptureBaseData& InBaseData, TSharedPtr<FMediaCaptureUserData, ESPMode::ThreadSafe> InUserData, const FMediaCaptureResourceData& InResourceData);
 
 private:
 	friend BlackmagicMediaCaptureHelpers::FBlackmagicMediaCaptureEventCallback;
@@ -86,7 +89,7 @@ private:
 	FFrameRate FrameRate;
 
 	/** Critical section for synchronizing access to the OutputChannel */
-	FCriticalSection RenderThreadCriticalSection;
+	FCriticalSection CopyingCriticalSection;
 
 	/** Event to wakeup When waiting for sync */
 	FEvent* WakeUpEvent;
