MediaIOFramework

Adding ability to specify the audio device in MediaCapture.

diff --git a/Engine/Plugins/Media/AjaMedia/Source/AjaMediaOutput/Private/AjaMediaCapture.cpp b/Engine/Plugins/Media/AjaMedia/Source/AjaMediaOutput/Private/AjaMediaCapture.cpp
--- a/Engine/Plugins/Media/AjaMedia/Source/AjaMediaOutput/Private/AjaMediaCapture.cpp
+++ b/Engine/Plugins/Media/AjaMedia/Source/AjaMediaOutput/Private/AjaMediaCapture.cpp
@@ -297,6 +297,11 @@ bool UAjaMediaCapture::UpdateRenderTargetImpl(UTextureRenderTarget2D* InRenderTa
 	return true;
 }
 
+bool UAjaMediaCapture::UpdateAudioDeviceImpl(const FAudioDeviceHandle& InAudioDeviceHandle)
+{
+	return CreateAudioOutput(InAudioDeviceHandle, Cast<UAjaMediaOutput>(MediaOutput));
+}
+
 void UAjaMediaCapture::StopCaptureImpl(bool bAllowPendingFrameToBeProcess)
 {
 	if (!bAllowPendingFrameToBeProcess)
@@ -468,14 +473,12 @@ bool UAjaMediaCapture::InitAJA(UAjaMediaOutput* InAjaMediaOutput)
 
 	bOutputAudio = InAjaMediaOutput->bOutputAudio;
 	
-	if (GEngine && bOutputAudio)
+	if (bOutputAudio)
 	{
-		UMediaIOCoreSubsystem::FCreateAudioOutputArgs Args;
-		Args.NumOutputChannels = static_cast<uint32>(InAjaMediaOutput->NumOutputAudioChannels);
-		Args.TargetFrameRate = FrameRate;
-		Args.MaxSampleLatency = InAjaMediaOutput->AudioBufferSize;
-		Args.OutputSampleRate = static_cast<uint32>(InAjaMediaOutput->AudioSampleRate);
-		AudioOutput = GEngine->GetEngineSubsystem<UMediaIOCoreSubsystem>()->CreateAudioOutput(Args);
+		if (!CreateAudioOutput(AudioDeviceHandle, InAjaMediaOutput))
+		{
+			UE_LOG(LogAjaMediaOutput, Error, TEXT("Failed to initialize audio output."));
+		}
 	}
 	
 	PixelFormat = InAjaMediaOutput->PixelFormat;
@@ -511,6 +514,22 @@ bool UAjaMediaCapture::InitAJA(UAjaMediaOutput* InAjaMediaOutput)
 	return true;
 }
 
+bool UAjaMediaCapture::CreateAudioOutput(const FAudioDeviceHandle& InAudioDeviceHandle, const UAjaMediaOutput* InAjaMediaOutput)
+{
+	if (GEngine && bOutputAudio && InAjaMediaOutput)
+	{
+		UMediaIOCoreSubsystem::FCreateAudioOutputArgs Args;
+		Args.NumOutputChannels = static_cast<int32>(InAjaMediaOutput->NumOutputAudioChannels);
+		Args.TargetFrameRate = FrameRate;
+		Args.MaxSampleLatency = Align(InAjaMediaOutput->AudioBufferSize, 4);
+		Args.OutputSampleRate = static_cast<uint32>(InAjaMediaOutput->AudioSampleRate);
+		Args.AudioDeviceHandle = InAudioDeviceHandle;
+		AudioOutput = GEngine->GetEngineSubsystem<UMediaIOCoreSubsystem>()->CreateAudioOutput(Args);
+		return AudioOutput.IsValid();
+	}
+	return false;
+}
+
 void UAjaMediaCapture::OnFrameCaptured_RenderingThread(const FCaptureBaseData& InBaseData, TSharedPtr<FMediaCaptureUserData, ESPMode::ThreadSafe> InUserData, void* InBuffer, int32 Width, int32 Height, int32 BytesPerRow)
 {
 	TRACE_CPUPROFILER_EVENT_SCOPE(UAjaMediaCapture::OnFrameCaptured_RenderingThread);
@@ -592,12 +611,19 @@ void UAjaMediaCapture::WaitForSync_AnyThread() const
 
 void UAjaMediaCapture::OutputAudio_AnyThread(const AJA::AJAOutputFrameBufferData& FrameBuffer) const
 {
-	if (bOutputAudio && AudioOutput)
+	if (bOutputAudio)
 	{
 		TRACE_CPUPROFILER_EVENT_SCOPE(UAjaMediaCapture::OnFrameCaptured_RenderingThread::SetAudio);
-
-		int32 NumSamplesToPull = FMath::RoundToInt32(48000.f * AudioOutput->NumInputChannels / FrameRate.AsDecimal());
-		TArray<int32> AudioSamples = AudioOutput->GetAudioSamples<int32>(NumSamplesToPull);
+		
+		// Take a local copy of the audio output in case it is switched from the main thread.
+		const TSharedPtr<FMediaIOAudioOutput> LocalAudioOutput = AudioOutput;
+		if(!LocalAudioOutput)
+		{
+			return;
+		}
+		
+		const int32 NumSamplesToPull = FMath::RoundToInt32(48000.f * LocalAudioOutput->NumInputChannels / FrameRate.AsDecimal());
+		TArray<int32> AudioSamples = LocalAudioOutput->GetAudioSamples<int32>(NumSamplesToPull);
 		OutputChannel->SetAudioFrameData(FrameBuffer, reinterpret_cast<uint8*>(AudioSamples.GetData()), AudioSamples.Num() * sizeof(int32));
 	}
 }
diff --git a/Engine/Plugins/Media/AjaMedia/Source/AjaMediaOutput/Public/AjaMediaCapture.h b/Engine/Plugins/Media/AjaMedia/Source/AjaMediaOutput/Public/AjaMediaCapture.h
--- a/Engine/Plugins/Media/AjaMedia/Source/AjaMediaOutput/Public/AjaMediaCapture.h
+++ b/Engine/Plugins/Media/AjaMedia/Source/AjaMediaOutput/Public/AjaMediaCapture.h
@@ -46,6 +46,7 @@ protected:
 	virtual bool PostInitializeCaptureViewport(TSharedPtr<FSceneViewport>& InSceneViewport) override;
 	virtual bool UpdateSceneViewportImpl(TSharedPtr<FSceneViewport>& InSceneViewport) override;
 	virtual bool UpdateRenderTargetImpl(UTextureRenderTarget2D* InRenderTarget) override;
+	virtual bool UpdateAudioDeviceImpl(const FAudioDeviceHandle& InAudioDeviceHandle) override;
 	virtual void StopCaptureImpl(bool bAllowPendingFrameToBeProcess) override;
 	virtual bool ShouldCaptureRHIResource() const override;
 	
@@ -64,6 +65,7 @@ private:
 
 private:
 	bool InitAJA(UAjaMediaOutput* InMediaOutput);
+	bool CreateAudioOutput(const FAudioDeviceHandle& InAudioDeviceHandle, const UAjaMediaOutput* InAjaMediaOutput);
 	void WaitForSync_AnyThread() const;
 	void OutputAudio_AnyThread(const AJA::AJAOutputFrameBufferData& FrameBuffer) const;
 	void ApplyViewportTextureAlpha(TSharedPtr<FSceneViewport> InSceneViewport);
diff --git a/Engine/Plugins/Media/BlackmagicMedia/Source/BlackmagicMediaOutput/Private/BlackmagicMediaCapture.cpp b/Engine/Plugins/Media/BlackmagicMedia/Source/BlackmagicMediaOutput/Private/BlackmagicMediaCapture.cpp
--- a/Engine/Plugins/Media/BlackmagicMedia/Source/BlackmagicMediaOutput/Private/BlackmagicMediaCapture.cpp
+++ b/Engine/Plugins/Media/BlackmagicMedia/Source/BlackmagicMediaOutput/Private/BlackmagicMediaCapture.cpp
@@ -386,6 +386,11 @@ bool UBlackmagicMediaCapture::UpdateRenderTargetImpl(UTextureRenderTarget2D* InR
 	return true;
 }
 
+bool UBlackmagicMediaCapture::UpdateAudioDeviceImpl(const FAudioDeviceHandle& InAudioDeviceHandle)
+{
+	return CreateAudioOutput(InAudioDeviceHandle, Cast<UBlackmagicMediaOutput>(MediaOutput));
+}
+
 void UBlackmagicMediaCapture::StopCaptureImpl(bool bAllowPendingFrameToBeProcess)
 {
 	if (!bAllowPendingFrameToBeProcess)
@@ -605,14 +610,12 @@ bool UBlackmagicMediaCapture::InitBlackmagic(UBlackmagicMediaOutput* InBlackmagi
 	bOutputAudio = InBlackmagicMediaOutput->bOutputAudio;
 	NumOutputChannels = static_cast<uint8>(InBlackmagicMediaOutput->OutputChannelCount);
 
-	if (GEngine && bOutputAudio)
+	if (bOutputAudio)
 	{
-		UMediaIOCoreSubsystem::FCreateAudioOutputArgs Args;
-		Args.NumOutputChannels = static_cast<int32>(ChannelOptions.NumAudioChannels);
-		Args.TargetFrameRate = FrameRate;
-		Args.MaxSampleLatency = Align(InBlackmagicMediaOutput->AudioBufferSize, 4);
-		Args.OutputSampleRate = static_cast<uint32>(InBlackmagicMediaOutput->AudioSampleRate);
-		AudioOutput = GEngine->GetEngineSubsystem<UMediaIOCoreSubsystem>()->CreateAudioOutput(Args);
+		if (!CreateAudioOutput(AudioDeviceHandle, InBlackmagicMediaOutput))
+		{
+			UE_LOG(LogBlackmagicMediaOutput, Error, TEXT("Failed to initialize audio output."));
+		}
 	}
 	
 	check(EventCallback == nullptr);
@@ -650,6 +653,22 @@ bool UBlackmagicMediaCapture::InitBlackmagic(UBlackmagicMediaOutput* InBlackmagi
 	return true;
 }
 
+bool UBlackmagicMediaCapture::CreateAudioOutput(const FAudioDeviceHandle& InAudioDeviceHandle, const UBlackmagicMediaOutput* InBlackmagicMediaOutput)
+{
+	if (GEngine && bOutputAudio && InBlackmagicMediaOutput)
+	{
+		UMediaIOCoreSubsystem::FCreateAudioOutputArgs Args;
+		Args.NumOutputChannels = static_cast<int32>(InBlackmagicMediaOutput->OutputChannelCount);
+		Args.TargetFrameRate = FrameRate;
+		Args.MaxSampleLatency = Align(InBlackmagicMediaOutput->AudioBufferSize, 4);
+		Args.OutputSampleRate = static_cast<uint32>(InBlackmagicMediaOutput->AudioSampleRate);
+		Args.AudioDeviceHandle = InAudioDeviceHandle;
+		AudioOutput = GEngine->GetEngineSubsystem<UMediaIOCoreSubsystem>()->CreateAudioOutput(Args);
+		return AudioOutput.IsValid();
+	}
+	return false;
+}
+
 void UBlackmagicMediaCapture::LockDMATexture_RenderThread(FTextureRHIRef InTexture)
 {
 	if (ShouldCaptureRHIResource() && TextureTransfer)
@@ -787,10 +806,17 @@ void UBlackmagicMediaCapture::WaitForSync_AnyThread()
 
 void UBlackmagicMediaCapture::OutputAudio_AnyThread(const FCaptureBaseData& InBaseData, const BlackmagicDesign::FTimecode& InTimecode)
 {
-	if (bOutputAudio && AudioOutput)
+	if (bOutputAudio)
 	{
 		TRACE_CPUPROFILER_EVENT_SCOPE(UBlackmagicMediaCapture::OutputAudio);
 
+		// Take a local copy of the audio output in case it is switched from the main thread.
+		const TSharedPtr<FMediaIOAudioOutput> LocalAudioOutput = AudioOutput;
+		if(!LocalAudioOutput)
+		{
+			return;
+		}
+
 		const double NewTimestamp = FPlatformTime::Seconds();
 		double CurrentFrameTime = NewTimestamp - OutputAudioTimestamp;
 		
@@ -802,8 +828,8 @@ void UBlackmagicMediaCapture::OutputAudio_AnyThread(const FCaptureBaseData& InBa
 
 		float FrameTimeRatio = CurrentFrameTime / TargetFrametime;
 
-		uint32_t NumSamplesToPull = FrameTimeRatio * AudioOutput->NumSamplesPerFrame;
-		NumSamplesToPull = FMath::Clamp<uint32_t>(NumSamplesToPull, 0, AudioOutput->NumSamplesPerFrame);
+		uint32_t NumSamplesToPull = FrameTimeRatio * LocalAudioOutput->NumSamplesPerFrame;
+		NumSamplesToPull = FMath::Clamp<uint32_t>(NumSamplesToPull, 0, LocalAudioOutput->NumSamplesPerFrame);
 
 		BlackmagicDesign::FAudioSamplesDescriptor AudioSamples;
 		AudioSamples.Timecode = InTimecode;
@@ -812,7 +838,7 @@ void UBlackmagicMediaCapture::OutputAudio_AnyThread(const FCaptureBaseData& InBa
 		
 		if (AudioBitDepth == EBlackmagicMediaOutputAudioBitDepth::Signed_32Bits)
 		{
-			TArray<int32> AudioBuffer = AudioOutput->GetAllAudioSamples<int32>();
+			TArray<int32> AudioBuffer = LocalAudioOutput->GetAllAudioSamples<int32>();
 			AudioSamples.AudioBuffer = reinterpret_cast<uint8_t*>(AudioBuffer.GetData());
 			AudioSamples.NumAudioSamples = AudioBuffer.Num() / NumOutputChannels;
 			AudioSamples.AudioBufferLength = AudioBuffer.Num() * sizeof(int32);
@@ -820,7 +846,7 @@ void UBlackmagicMediaCapture::OutputAudio_AnyThread(const FCaptureBaseData& InBa
 		}
 		else if (AudioBitDepth == EBlackmagicMediaOutputAudioBitDepth::Signed_16Bits)
 		{
-			TArray<int16> AudioBuffer = AudioOutput->GetAllAudioSamples<int16>();
+			TArray<int16> AudioBuffer = LocalAudioOutput->GetAllAudioSamples<int16>();
 			AudioSamples.AudioBuffer = reinterpret_cast<uint8_t*>(AudioBuffer.GetData());
 			AudioSamples.NumAudioSamples = AudioBuffer.Num() / NumOutputChannels;
 			AudioSamples.AudioBufferLength = AudioBuffer.Num() * sizeof(int16);
diff --git a/Engine/Plugins/Media/BlackmagicMedia/Source/BlackmagicMediaOutput/Public/BlackmagicMediaCapture.h b/Engine/Plugins/Media/BlackmagicMedia/Source/BlackmagicMediaOutput/Public/BlackmagicMediaCapture.h
--- a/Engine/Plugins/Media/BlackmagicMedia/Source/BlackmagicMediaOutput/Public/BlackmagicMediaCapture.h
+++ b/Engine/Plugins/Media/BlackmagicMedia/Source/BlackmagicMediaOutput/Public/BlackmagicMediaCapture.h
@@ -47,6 +47,7 @@ protected:
 	virtual bool PostInitializeCaptureViewport(TSharedPtr<FSceneViewport>& InSceneViewport) override;
 	virtual bool UpdateSceneViewportImpl(TSharedPtr<FSceneViewport>& InSceneViewport) override;
 	virtual bool UpdateRenderTargetImpl(UTextureRenderTarget2D* InRenderTarget) override;
+	virtual bool UpdateAudioDeviceImpl(const FAudioDeviceHandle& InAudioDeviceHandle) override;
 	virtual void StopCaptureImpl(bool bAllowPendingFrameToBeProcess) override;
 	virtual bool ShouldCaptureRHIResource() const override;
 
@@ -63,6 +64,7 @@ private:
 
 private:
 	bool InitBlackmagic(UBlackmagicMediaOutput* InMediaOutput);
+	bool CreateAudioOutput(const FAudioDeviceHandle& InAudioDeviceHandle, const UBlackmagicMediaOutput* InBlackmagicMediaOutput);
 	void WaitForSync_AnyThread();
 	void OutputAudio_AnyThread(const FCaptureBaseData& InBaseData, const BlackmagicDesign::FTimecode& Timecode);
 	void ApplyViewportTextureAlpha(TSharedPtr<FSceneViewport> InSceneViewport);
diff --git a/Engine/Plugins/Media/MediaIOFramework/Source/MediaIOCore/Private/MediaCapture.cpp b/Engine/Plugins/Media/MediaIOFramework/Source/MediaIOCore/Private/MediaCapture.cpp
--- a/Engine/Plugins/Media/MediaIOFramework/Source/MediaIOCore/Private/MediaCapture.cpp
+++ b/Engine/Plugins/Media/MediaIOFramework/Source/MediaIOCore/Private/MediaCapture.cpp
@@ -1347,6 +1347,20 @@ void UMediaCapture::SetMediaOutput(UMediaOutput* InMediaOutput)
 	}
 }
 
+bool UMediaCapture::SetCaptureAudioDevice(const FAudioDeviceHandle& InAudioDeviceHandle)
+{
+	bool bSuccess = true;
+	if (CaptureSource)
+	{
+		bSuccess = UpdateAudioDeviceImpl(InAudioDeviceHandle);
+	}
+	if (bSuccess)
+	{
+		AudioDeviceHandle = InAudioDeviceHandle;
+	}
+	return bSuccess;
+}
+
 void UMediaCapture::CaptureImmediate_RenderThread(FRDGBuilder& GraphBuilder, FRHITexture* InSourceTexture)
 {
 	UE::MediaCaptureData::FCaptureFrameArgs CaptureArgs{GraphBuilder};
diff --git a/Engine/Plugins/Media/MediaIOFramework/Source/MediaIOCore/Private/MediaIOCoreAudioOutput.cpp b/Engine/Plugins/Media/MediaIOFramework/Source/MediaIOCore/Private/MediaIOCoreAudioOutput.cpp
--- a/Engine/Plugins/Media/MediaIOFramework/Source/MediaIOCore/Private/MediaIOCoreAudioOutput.cpp
+++ b/Engine/Plugins/Media/MediaIOFramework/Source/MediaIOCore/Private/MediaIOCoreAudioOutput.cpp
@@ -54,23 +54,14 @@ Audio::FAlignedFloatBuffer FMediaIOAudioOutput::GetFloatBuffer(uint32 NumSamples
 	return FloatBuffer;
 }
 
-FMediaIOAudioCapture::FMediaIOAudioCapture()
+FMediaIOAudioCapture::FMediaIOAudioCapture(const FAudioDeviceHandle& InAudioDeviceHandle)
 {
-	RegisterMainAudioDevice();
-#if WITH_EDITOR
-	FEditorDelegates::PostPIEStarted.AddRaw(this, &FMediaIOAudioCapture::OnPIEStarted);
-	FEditorDelegates::PrePIEEnded.AddRaw(this, &FMediaIOAudioCapture::OnPIEEnded);
-#endif
+	RegisterAudioDevice(InAudioDeviceHandle);
 }
 
 FMediaIOAudioCapture::~FMediaIOAudioCapture()
 {
-#if WITH_EDITOR
-	FEditorDelegates::EndPIE.RemoveAll(this);
-	FEditorDelegates::PostPIEStarted.RemoveAll(this);
-#endif
-
-	UnregisterMainAudioDevice();
+	UnregisterAudioDevice();
 }
 
 void FMediaIOAudioCapture::OnNewSubmixBuffer(const USoundSubmix* InOwningSubmix, float* InAudioData, int32 InNumSamples, int32 InNumChannels, const int32 InSampleRate, double InAudioClock)
@@ -131,8 +122,72 @@ TSharedPtr<FMediaIOAudioOutput> FMediaIOAudioCapture::CreateAudioOutput(int32 In
 	return nullptr;
 }
 
+void FMediaIOAudioCapture::RegisterBufferListener(FAudioDevice* AudioDevice)
+{
+	if (AudioDevice && AudioDevice->IsAudioMixerEnabled())
+	{
+		const Audio::FMixerDevice* MixerDevice = static_cast<Audio::FMixerDevice*>(AudioDevice);
+		NumChannels = MixerDevice->GetDeviceOutputChannels();
+		SampleRate = MixerDevice->GetSampleRate();
+		PrimarySubmixName = *GetDefault<UAudioSettings>()->MasterSubmix.GetAssetName();
+		AudioDevice->RegisterSubmixBufferListener(this);
+	}
+}
+
+void FMediaIOAudioCapture::UnregisterBufferListener(FAudioDevice* AudioDevice)
+{
+	if (AudioDevice)
+	{
+		AudioDevice->UnregisterSubmixBufferListener(this);
+	}
+}
+
+void FMediaIOAudioCapture::RegisterAudioDevice(const FAudioDeviceHandle& InAudioDeviceHandle)
+{
+	// Can only be registered to one device.
+	UnregisterAudioDevice();
+
+	if (InAudioDeviceHandle.IsValid())
+	{
+		RegisterBufferListener(InAudioDeviceHandle.GetAudioDevice());
+		RegisteredDeviceId = InAudioDeviceHandle.GetDeviceID();
+	}
+}
+
+void FMediaIOAudioCapture::UnregisterAudioDevice()
+{
+	if (RegisteredDeviceId != INDEX_NONE)
+	{
+		if (FAudioDevice* RegisteredDevice = GEngine->GetAudioDeviceManager()->GetAudioDeviceRaw(RegisteredDeviceId))
+		{
+			UnregisterBufferListener(RegisteredDevice);
+		}
+		RegisteredDeviceId = INDEX_NONE;
+	}
+}
+
+FMainMediaIOAudioCapture::FMainMediaIOAudioCapture()
+	: FMediaIOAudioCapture(FAudioDeviceHandle())
+{
+	RegisterMainAudioDevice();
 #if WITH_EDITOR
-void FMediaIOAudioCapture::OnPIEStarted(const bool)
+	FEditorDelegates::PostPIEStarted.AddRaw(this, &FMainMediaIOAudioCapture::OnPIEStarted);
+	FEditorDelegates::PrePIEEnded.AddRaw(this, &FMainMediaIOAudioCapture::OnPIEEnded);
+#endif
+}
+
+FMainMediaIOAudioCapture::~FMainMediaIOAudioCapture()
+{
+#if WITH_EDITOR
+	FEditorDelegates::EndPIE.RemoveAll(this);
+	FEditorDelegates::PostPIEStarted.RemoveAll(this);
+#endif
+
+	UnregisterMainAudioDevice();
+}
+
+#if WITH_EDITOR
+void FMainMediaIOAudioCapture::OnPIEStarted(const bool)
 {
 	UnregisterMainAudioDevice();
 	
@@ -146,7 +201,7 @@ void FMediaIOAudioCapture::OnPIEStarted(const bool)
 	}
 }
 
-void FMediaIOAudioCapture::OnPIEEnded(const bool)
+void FMainMediaIOAudioCapture::OnPIEEnded(const bool)
 {
 	if (GEditor)
 	{
@@ -160,35 +215,14 @@ void FMediaIOAudioCapture::OnPIEEnded(const bool)
 }
 #endif
 
-void FMediaIOAudioCapture::RegisterMainAudioDevice()
+void FMainMediaIOAudioCapture::RegisterMainAudioDevice()
 {
 	RegisterBufferListener(GEngine->GetMainAudioDeviceRaw());
 }
 
-void FMediaIOAudioCapture::UnregisterMainAudioDevice()
+void FMainMediaIOAudioCapture::UnregisterMainAudioDevice()
 {
 	UnregisterBufferListener(GEngine->GetMainAudioDeviceRaw());
 }
 
-void FMediaIOAudioCapture::RegisterBufferListener(FAudioDevice* AudioDevice)
-{
-	if (AudioDevice && AudioDevice->IsAudioMixerEnabled())
-	{
-		Audio::FMixerDevice* MixerDevice = static_cast<Audio::FMixerDevice*>(AudioDevice);
-		NumChannels = MixerDevice->GetDeviceOutputChannels();
-		SampleRate = MixerDevice->GetSampleRate();
-		PrimarySubmixName = *GetDefault<UAudioSettings>()->MasterSubmix.GetAssetName();
-		AudioDevice->RegisterSubmixBufferListener(this);
-	}
-}
-
-void FMediaIOAudioCapture::UnregisterBufferListener(FAudioDevice* AudioDevice)
-{
-	if (AudioDevice)
-	{
-		AudioDevice->UnregisterSubmixBufferListener(this);
-	}
-}
-
-
 #undef LOCTEXT_NAMESPACE //MediaIOAudioOutput
\ No newline at end of file
diff --git a/Engine/Plugins/Media/MediaIOFramework/Source/MediaIOCore/Private/MediaIOCoreSubsystem.cpp b/Engine/Plugins/Media/MediaIOFramework/Source/MediaIOCore/Private/MediaIOCoreSubsystem.cpp
--- a/Engine/Plugins/Media/MediaIOFramework/Source/MediaIOCore/Private/MediaIOCoreSubsystem.cpp
+++ b/Engine/Plugins/Media/MediaIOFramework/Source/MediaIOCore/Private/MediaIOCoreSubsystem.cpp
@@ -1,16 +1,52 @@
 // Copyright Epic Games, Inc. All Rights Reserved.
 
 #include "MediaIOCoreSubsystem.h"
+#include "AudioDeviceManager.h"
 
 #include UE_INLINE_GENERATED_CPP_BY_NAME(MediaIOCoreSubsystem)
 
-TSharedPtr<FMediaIOAudioOutput> UMediaIOCoreSubsystem::CreateAudioOutput(const FCreateAudioOutputArgs& Args)
+void UMediaIOCoreSubsystem::Initialize(FSubsystemCollectionBase& InCollection)
 {
-	if (!MediaIOCapture)
+	DeviceDestroyedHandle = FAudioDeviceManagerDelegates::OnAudioDeviceDestroyed.AddUObject(this, &UMediaIOCoreSubsystem::OnAudioDeviceDestroyed);
+}
+
+void UMediaIOCoreSubsystem::Deinitialize()
+{
+	FAudioDeviceManagerDelegates::OnAudioDeviceDestroyed.Remove(DeviceDestroyedHandle);
+}
+
+TSharedPtr<FMediaIOAudioOutput> UMediaIOCoreSubsystem::CreateAudioOutput(const FCreateAudioOutputArgs& InArgs)
+{
+	FMediaIOAudioCapture* MediaIOAudioCapture = nullptr;
+
+	if (InArgs.AudioDeviceHandle.IsValid())
 	{
-		MediaIOCapture = MakeUnique<FMediaIOAudioCapture>();
+		if (const TUniquePtr<FMediaIOAudioCapture>* FoundMediaIOAudioCapture = MediaIOAudioCaptures.Find(InArgs.AudioDeviceHandle.GetDeviceID()))
+		{
+			MediaIOAudioCapture = FoundMediaIOAudioCapture->Get();
+		}
+		else
+		{
+			MediaIOAudioCapture = MediaIOAudioCaptures.Add(InArgs.AudioDeviceHandle.GetDeviceID(), MakeUnique<FMediaIOAudioCapture>(InArgs.AudioDeviceHandle)).Get();
+		}
 	}
+	else
+	{
+		// Fallback using the main audio device.
+		if (!MainMediaIOAudioCapture)
+		{
+			MainMediaIOAudioCapture = MakeUnique<FMainMediaIOAudioCapture>();
+		}
+		MediaIOAudioCapture = MainMediaIOAudioCapture.Get();
+	}	
 
-	return MediaIOCapture->CreateAudioOutput(Args.NumOutputChannels, Args.TargetFrameRate, Args.MaxSampleLatency, Args.OutputSampleRate);
+	return MediaIOAudioCapture->CreateAudioOutput(InArgs.NumOutputChannels, InArgs.TargetFrameRate, InArgs.MaxSampleLatency, InArgs.OutputSampleRate);
 }
 
+void UMediaIOCoreSubsystem::OnAudioDeviceDestroyed(Audio::FDeviceId InAudioDeviceId)
+{
+	if (MediaIOAudioCaptures.Contains(InAudioDeviceId))
+	{
+		MediaIOAudioCaptures.Remove(InAudioDeviceId);
+	}
+}
diff --git a/Engine/Plugins/Media/MediaIOFramework/Source/MediaIOCore/Public/MediaCapture.h b/Engine/Plugins/Media/MediaIOFramework/Source/MediaIOCore/Public/MediaCapture.h
--- a/Engine/Plugins/Media/MediaIOFramework/Source/MediaIOCore/Public/MediaCapture.h
+++ b/Engine/Plugins/Media/MediaIOFramework/Source/MediaIOCore/Public/MediaCapture.h
@@ -6,6 +6,7 @@
 #include "UObject/Object.h"
 
 #include <atomic>
+#include "AudioDeviceHandle.h"
 #include "Containers/SpscQueue.h"
 #include "HAL/CriticalSection.h"
 #include "MediaOutput.h"
@@ -303,6 +304,12 @@ public:
 	UFUNCTION(BlueprintCallable, Category = "Media|Output")
 	void SetMediaOutput(UMediaOutput* InMediaOutput);
 
+	/** Set the audio device to capture. */
+	bool SetCaptureAudioDevice(const FAudioDeviceHandle& InAudioDeviceHandle);
+
+	/** Get the audio device to capture. */
+	const FAudioDeviceHandle& GetCaptureAudioDevice() const { return AudioDeviceHandle; }
+	
 	/** Get the desired size of the current capture. */
 	UFUNCTION(BlueprintCallable, Category = "Media|Output")
 	FIntPoint GetDesiredSize() const { return DesiredSize; }
@@ -341,6 +348,7 @@ public:
 	
 	virtual bool UpdateSceneViewportImpl(TSharedPtr<FSceneViewport>& InSceneViewport) { return true; }
 	virtual bool UpdateRenderTargetImpl(UTextureRenderTarget2D* InRenderTarget) { return true; }
+	virtual bool UpdateAudioDeviceImpl(const FAudioDeviceHandle& InAudioDeviceHandle) { return true; }
 
 	static const TSet<EPixelFormat>& GetSupportedRgbaSwizzleFormats();
 
@@ -541,6 +549,9 @@ protected:
 	/** MediaOutput associated with this capture */
 	UPROPERTY(Transient)
 	TObjectPtr<UMediaOutput> MediaOutput;
+
+	/** Audio device to be captured. If none, main engine's audio device is captured. */
+	FAudioDeviceHandle AudioDeviceHandle;
 	
 	/** Output size of the media capture. If a conversion is done, this resolution might be different than source resolution. i.e. 1080p RGB might be 960x1080 in YUV */
 	FIntPoint DesiredOutputSize = FIntPoint(1920, 1080);
diff --git a/Engine/Plugins/Media/MediaIOFramework/Source/MediaIOCore/Public/MediaIOCoreAudioOutput.h b/Engine/Plugins/Media/MediaIOFramework/Source/MediaIOCore/Public/MediaIOCoreAudioOutput.h
--- a/Engine/Plugins/Media/MediaIOFramework/Source/MediaIOCore/Public/MediaIOCoreAudioOutput.h
+++ b/Engine/Plugins/Media/MediaIOFramework/Source/MediaIOCore/Public/MediaIOCoreAudioOutput.h
@@ -11,6 +11,7 @@
 #endif
 
 class FAudioDevice;
+class FAudioDeviceHandle;
 
 class MEDIAIOCORE_API FMediaIOAudioOutput
 { 
@@ -138,12 +139,12 @@ private:
 };
 
 /**
- * Handles capturing capturing audio samples rendered by the engine and dispatching them to outputs.
+ * Handles capturing audio samples rendered by the engine and dispatching them to outputs.
  */
 class FMediaIOAudioCapture : public ISubmixBufferListener
 {
 public:
-	FMediaIOAudioCapture();
+	FMediaIOAudioCapture(const FAudioDeviceHandle& InAudioDeviceHandle);
 	virtual ~FMediaIOAudioCapture();
 
 	//~ ISubmixBufferListener interface
@@ -152,17 +153,18 @@ public:
 	/** Create an audio output that will receive audio samples. */
 	TSharedPtr<FMediaIOAudioOutput> CreateAudioOutput(int32 InNumOutputChannels, FFrameRate InTargetFrameRate, uint32 InMaxSampleLatency, uint32 InOutputSampleRate);
 
-private:
-#if WITH_EDITOR
-	void OnPIEStarted(const bool);
-	void OnPIEEnded(const bool);
-#endif
-
-	void RegisterMainAudioDevice();
-	void UnregisterMainAudioDevice();
+protected:
 	void RegisterBufferListener(FAudioDevice* AudioDevice);
 	void UnregisterBufferListener(FAudioDevice* AudioDevice);
+
+private:
+	void RegisterAudioDevice(const FAudioDeviceHandle& InAudioDeviceHandle);
+	void UnregisterAudioDevice();
+
 private:
+	/** Audio device Id this buffer listener is registered to. */
+	Audio::FDeviceId RegisteredDeviceId = INDEX_NONE;
+	
 	/** Sample rate on the engine side. */ 
 	uint32 SampleRate = 0;
 
@@ -175,3 +177,25 @@ private:
 	/** Utility that allows pushing audio samples to multiple outputs. */
 	Audio::FPatchSplitter AudioSplitter;
 };
+
+/**
+ * Audio capture that automatically registers to the main engine device.
+ * Also, handles automatically registering to the current PIE world's audio device.
+ * This audio capture is used by default if no audio device handle is specified when media capture
+ * creates the audio output.
+ */
+class FMainMediaIOAudioCapture : public FMediaIOAudioCapture
+{
+public:
+	FMainMediaIOAudioCapture();
+	virtual ~FMainMediaIOAudioCapture();
+
+private:
+#if WITH_EDITOR
+	void OnPIEStarted(const bool);
+	void OnPIEEnded(const bool);
+#endif
+
+	void RegisterMainAudioDevice();
+	void UnregisterMainAudioDevice();
+};
\ No newline at end of file
diff --git a/Engine/Plugins/Media/MediaIOFramework/Source/MediaIOCore/Public/MediaIOCoreSubsystem.h b/Engine/Plugins/Media/MediaIOFramework/Source/MediaIOCore/Public/MediaIOCoreSubsystem.h
--- a/Engine/Plugins/Media/MediaIOFramework/Source/MediaIOCore/Public/MediaIOCoreSubsystem.h
+++ b/Engine/Plugins/Media/MediaIOFramework/Source/MediaIOCore/Public/MediaIOCoreSubsystem.h
@@ -2,6 +2,7 @@
 
 #pragma once
 
+#include "AudioDeviceHandle.h"
 #include "MediaIOCoreAudioOutput.h"
 #include "Subsystems/EngineSubsystem.h"
 
@@ -17,16 +18,29 @@ public:
 		FFrameRate TargetFrameRate; 
 		uint32 MaxSampleLatency = 0;
 		uint32 OutputSampleRate = 0;
+		FAudioDeviceHandle AudioDeviceHandle;
 	};
 
 public:
 	GENERATED_BODY()
 
+	//~ Begin UEngineSubsystem Interface
+	virtual void Initialize(FSubsystemCollectionBase& InCollection) override;
+	virtual void Deinitialize() override;
+	//~ End UEngineSubsystem Interface
+
 	/**
 	 * Create an audio output that allows getting audio that was accumulated during the last frame. 
 	 */
-	TSharedPtr<FMediaIOAudioOutput> CreateAudioOutput(const FCreateAudioOutputArgs& Args);
+	TSharedPtr<FMediaIOAudioOutput> CreateAudioOutput(const FCreateAudioOutputArgs& InArgs);
+
+private:
+	void OnAudioDeviceDestroyed(Audio::FDeviceId InAudioDeviceId);
 
 private:
-	TUniquePtr<FMediaIOAudioCapture> MediaIOCapture;
+	TUniquePtr<FMediaIOAudioCapture> MainMediaIOAudioCapture;
+	
+	TMap<Audio::FDeviceId, TUniquePtr<FMediaIOAudioCapture>> MediaIOAudioCaptures;
+
+	FDelegateHandle DeviceDestroyedHandle;
 };
