Add missing file from previous CL (25508091) Media input deinterlacing improvements

diff --git a/Engine/Plugins/Media/BlackmagicMedia/Source/BlackmagicMedia/Private/Player/BlackmagicMediaPlayer.h b/Engine/Plugins/Media/BlackmagicMedia/Source/BlackmagicMedia/Private/Player/BlackmagicMediaPlayer.h
--- a/Engine/Plugins/Media/BlackmagicMedia/Source/BlackmagicMedia/Private/Player/BlackmagicMediaPlayer.h
+++ b/Engine/Plugins/Media/BlackmagicMedia/Source/BlackmagicMedia/Private/Player/BlackmagicMediaPlayer.h
@@ -87,20 +87,22 @@ public:
 	virtual bool IsHardwareReady() const override;
 
 protected:
-	/** Setup our different channels with the current set of settings */
+	//~ Begin FMediaIOCorePlayerBase interface
 	virtual void SetupSampleChannels() override;
-
 	virtual uint32 GetNumVideoFrameBuffers() const override 
 	{
 		return MaxNumVideoFrameBuffer;
 	}
-
 	virtual EMediaIOCoreColorFormat GetColorFormat() const override
 	{
 		return BlackmagicColorFormat == EBlackmagicMediaSourceColorFormat::YUV8 ? EMediaIOCoreColorFormat::YUV8 : EMediaIOCoreColorFormat::YUV10;
 	}
-
 	virtual void AddVideoSample(const TSharedRef<FMediaIOCoreTextureSampleBase>& InSample) override;
+	virtual TSharedPtr<FMediaIOCoreTextureSampleBase> AcquireSample_AnyThread() const override
+	{
+		return TextureSamplePool->AcquireShared();
+	}
+	//~ End FMediaIOCorePlayerBase interface
 
 private:
 
