Add HDR support for blackmagic media output

diff --git a/Engine/Plugins/Media/BlackmagicMedia/Source/BlackmagicCore/Private/BlackmagicInputChannel.cpp b/Engine/Plugins/Media/BlackmagicMedia/Source/BlackmagicCore/Private/BlackmagicInputChannel.cpp
--- a/Engine/Plugins/Media/BlackmagicMedia/Source/BlackmagicCore/Private/BlackmagicInputChannel.cpp
+++ b/Engine/Plugins/Media/BlackmagicMedia/Source/BlackmagicCore/Private/BlackmagicInputChannel.cpp
@@ -381,6 +381,14 @@ namespace BlackmagicDesign
 						bmdDeckLinkFrameMetadataHDRDisplayPrimariesBlueX);
 					FrameInfo.HDRMetaData.DisplayPrimariesBlueY = GetFloatMetaData(MetadataExtensions,
 						bmdDeckLinkFrameMetadataHDRDisplayPrimariesBlueY);
+					FrameInfo.HDRMetaData.MaxDisplayLuminance = GetFloatMetaData(MetadataExtensions,
+						bmdDeckLinkFrameMetadataHDRMaxDisplayMasteringLuminance);
+					FrameInfo.HDRMetaData.MinDisplayLuminance = GetFloatMetaData(MetadataExtensions,
+						bmdDeckLinkFrameMetadataHDRMinDisplayMasteringLuminance);
+					FrameInfo.HDRMetaData.MaxContentLightLevel = GetFloatMetaData(MetadataExtensions,
+					bmdDeckLinkFrameMetadataHDRMaximumContentLightLevel);
+					FrameInfo.HDRMetaData.MaxFrameAverageLightLevel = GetFloatMetaData(MetadataExtensions,
+						bmdDeckLinkFrameMetadataHDRMaximumFrameAverageLightLevel);
 
 					MetadataExtensions->Release();
 				}
@@ -421,6 +429,18 @@ namespace BlackmagicDesign
 						case bmdDeckLinkFrameMetadataHDRWhitePointY:
 							UE_LOG(LogBlackmagicCore, Warning,TEXT("Could not get meta data HDR white point Y."));
 							break;
+						case bmdDeckLinkFrameMetadataHDRMaxDisplayMasteringLuminance:
+							UE_LOG(LogBlackmagicCore, Warning,TEXT("Could not get meta data HDR maximum mastering luminance."));
+							break;
+						case bmdDeckLinkFrameMetadataHDRMinDisplayMasteringLuminance:
+							UE_LOG(LogBlackmagicCore, Warning,TEXT("Could not get meta data HDR minimum mastering luminance."));
+							break;
+						case bmdDeckLinkFrameMetadataHDRMaximumContentLightLevel:
+							UE_LOG(LogBlackmagicCore, Warning,TEXT("Could not get meta data HDR maximum content light level."));
+							break;
+						case bmdDeckLinkFrameMetadataHDRMaximumFrameAverageLightLevel:
+							UE_LOG(LogBlackmagicCore, Warning,TEXT("Could not get meta data HDR maximum frame average light level."));
+							break;
 					}
 				}
 			}
diff --git a/Engine/Plugins/Media/BlackmagicMedia/Source/BlackmagicCore/Private/BlackmagicLib.cpp b/Engine/Plugins/Media/BlackmagicMedia/Source/BlackmagicCore/Private/BlackmagicLib.cpp
--- a/Engine/Plugins/Media/BlackmagicMedia/Source/BlackmagicCore/Private/BlackmagicLib.cpp
+++ b/Engine/Plugins/Media/BlackmagicMedia/Source/BlackmagicCore/Private/BlackmagicLib.cpp
@@ -63,6 +63,27 @@ namespace BlackmagicDesign
 		return DeviceIndex == Other.DeviceIndex;
 	}
 
+	/* FHDRMetaData
+	*****************************************************************************/
+	FHDRMetaData::FHDRMetaData()
+		: bIsAvailable(false)
+		, ColorSpace(EHDRMetaDataColorspace::Rec601)
+		, EOTF(EHDRMetaDataEOTF::SDR)
+		, WhitePointX(0.0)
+		, WhitePointY(0.0)
+		, DisplayPrimariesRedX(0.0)
+		, DisplayPrimariesRedY(0.0)
+		, DisplayPrimariesGreenX(0.0)
+		, DisplayPrimariesGreenY(0.0)
+		, DisplayPrimariesBlueX(0.0)
+		, DisplayPrimariesBlueY(0.0)
+		, MaxDisplayLuminance(0)
+		, MinDisplayLuminance(0)
+        , MaxContentLightLevel(0)
+        , MaxFrameAverageLightLevel(0)
+	
+	{ }
+
 	/* FInputChannelOptions
 	*****************************************************************************/
 	FInputChannelOptions::FInputChannelOptions()
@@ -108,20 +129,6 @@ namespace BlackmagicDesign
 		, AudioBuffer(nullptr)
 	{ }
 
-	IInputEventCallback::FFrameReceivedInfo::FHDRMetaData::FHDRMetaData()
-		: bIsAvailable(false)
-		, ColorSpace(EHDRMetaDataColorspace::Rec601)
-		, EOTF(EHDRMetaDataEOTF::SDR)
-		, WhitePointX(0.0)
-		, WhitePointY(0.0)
-		, DisplayPrimariesRedX(0.0)
-		, DisplayPrimariesRedY(0.0)
-		, DisplayPrimariesGreenX(0.0)
-		, DisplayPrimariesGreenY(0.0)
-		, DisplayPrimariesBlueX(0.0)
-		, DisplayPrimariesBlueY(0.0)
-	{ }
-
 	/* IOutputEventCallback 
 	*****************************************************************************/
 	IOutputEventCallback::~IOutputEventCallback()
diff --git a/Engine/Plugins/Media/BlackmagicMedia/Source/BlackmagicCore/Private/BlackmagicOutputChannel.cpp b/Engine/Plugins/Media/BlackmagicMedia/Source/BlackmagicCore/Private/BlackmagicOutputChannel.cpp
--- a/Engine/Plugins/Media/BlackmagicMedia/Source/BlackmagicCore/Private/BlackmagicOutputChannel.cpp
+++ b/Engine/Plugins/Media/BlackmagicMedia/Source/BlackmagicCore/Private/BlackmagicOutputChannel.cpp
@@ -110,8 +110,7 @@ namespace BlackmagicDesign
 		const uint32_t FOutputChannel::FOutputFrame::InvalidFrameIdentifier = static_cast<uint32_t>(-1);
 
 		FOutputChannel::FOutputFrame::FOutputFrame()
-			: BlackmagicFrame(nullptr)
-			, InFlightCount(0)
+			: InFlightCount(0)
 			, bVideoLineFilled(false)
 			, bVideoF2LineFilled(false) // used in interlaced
 			, bVideoScheduled(false)
@@ -121,10 +120,6 @@ namespace BlackmagicDesign
 
 		FOutputChannel::FOutputFrame::~FOutputFrame()
 		{
-			if (BlackmagicFrame)
-			{
-				BlackmagicFrame->Release();
-			}
 		}
 
 		void FOutputChannel::FOutputFrame::Clear()
@@ -432,7 +427,7 @@ namespace BlackmagicDesign
 					if (Result == S_OK)
 					{
 						FOutputFrame* NewFrame = new FOutputFrame();
-						NewFrame->BlackmagicFrame = NewBlackmagicFrame;
+						NewFrame->BlackmagicFrame = ReferencePtr<FDecklinkVideoFrame>(new FDecklinkVideoFrame(NewBlackmagicFrame, ChannelOptions.HDRMetadata));
 						AllFrames.push_back(NewFrame);
 						FrameReadyToWrite.push_back(NewFrame);
 					}
@@ -605,7 +600,7 @@ namespace BlackmagicDesign
 
 			const BMDTimeValue Duration = ChannelOptions.FormatInfo.FrameRateDenominator;
 			const BMDTimeScale Scale = ChannelOptions.FormatInfo.FrameRateNumerator;
-			HRESULT ScheduleVideoResult = DeckLinkOutput->ScheduleVideoFrame(InFrame->BlackmagicFrame, OutputTime, Duration, Scale);
+			HRESULT ScheduleVideoResult = DeckLinkOutput->ScheduleVideoFrame(static_cast<IDeckLinkVideoFrame*>(InFrame->BlackmagicFrame.Get()), OutputTime, Duration, Scale);
 
 #if PLATFORM_WINDOWS
 			// Don't call end sync complete if re-sending frame.
@@ -1560,7 +1555,7 @@ namespace BlackmagicDesign
 
 			return nullptr;
 		}
-		 
+
 		void FOutputChannel::PushWhenFrameReady(FOutputFrame* InFrame)
 		{
 			std::lock_guard<std::mutex> Lock(FrameLock);
@@ -1619,7 +1614,247 @@ namespace BlackmagicDesign
 			}
 			return true;
 		}
-}
+
+		HRESULT FDecklinkVideoFrame::QueryInterface(REFIID RIID, LPVOID* PPV)
+		{
+			HRESULT Result = S_OK;
+
+			if (PPV == nullptr)
+			{
+				return E_INVALIDARG;
+			}
+
+			if (IsEqualIID(RIID, IID_IUnknown))
+			{
+				*PPV = static_cast<IDeckLinkMutableVideoFrame*>(this);
+				AddRef();
+			}
+			else if (IsEqualIID(RIID, IID_IDeckLinkVideoFrameMetadataExtensions) && HDRMetadata.bIsAvailable)
+			{
+				*PPV = static_cast<IDeckLinkVideoFrameMetadataExtensions*>(this);
+				AddRef();
+			}
+			else
+			{
+				*PPV = nullptr;
+				Result = E_NOINTERFACE;
+			}
+
+			return Result;
+		}
+
+		ULONG FDecklinkVideoFrame::AddRef()
+		{
+			return ++RefCount;
+		}
+
+		ULONG FDecklinkVideoFrame::Release()
+		{
+			const ULONG NewRefValue = --RefCount;
+
+			if (NewRefValue == 0)
+			{
+				delete this;
+			}
+
+			return NewRefValue;
+		}
+
+		BMDFrameFlags FDecklinkVideoFrame::GetFlags(void)
+		{
+			if (HDRMetadata.bIsAvailable)
+			{
+				return WrappedVideoFrame->GetFlags() | bmdFrameContainsHDRMetadata;
+			}
+			return WrappedVideoFrame->GetFlags();
+		}
+
+		long FDecklinkVideoFrame::GetWidth(void)
+		{
+			return WrappedVideoFrame->GetWidth();
+		}
+
+		long FDecklinkVideoFrame::GetHeight(void)
+		{
+			return WrappedVideoFrame->GetHeight();
+		}
+
+		long FDecklinkVideoFrame::GetRowBytes(void)
+		{
+			return WrappedVideoFrame->GetRowBytes();
+		}
+
+		BMDPixelFormat FDecklinkVideoFrame::GetPixelFormat(void)
+		{
+			return WrappedVideoFrame->GetPixelFormat();
+		}
+
+		HRESULT FDecklinkVideoFrame::GetBytes(void** Buffer)
+		{
+			return WrappedVideoFrame->GetBytes(Buffer);
+		}
+
+		HRESULT FDecklinkVideoFrame::GetTimecode(BMDTimecodeFormat Format, IDeckLinkTimecode** Timecode)
+		{
+			return WrappedVideoFrame->GetTimecode(Format, Timecode);
+		}
+
+		HRESULT FDecklinkVideoFrame::GetAncillaryData(IDeckLinkVideoFrameAncillary** Ancillary)
+		{
+			return WrappedVideoFrame->GetAncillaryData(Ancillary);
+		}
+
+		HRESULT FDecklinkVideoFrame::SetFlags(BMDFrameFlags NewFlags)
+		{
+			return WrappedVideoFrame->SetFlags(NewFlags);
+		}
+
+		HRESULT FDecklinkVideoFrame::SetTimecode(BMDTimecodeFormat Format, IDeckLinkTimecode* Timecode)
+		{
+			return WrappedVideoFrame->SetTimecode(Format, Timecode);
+		}
+
+		HRESULT FDecklinkVideoFrame::SetTimecodeFromComponents(BMDTimecodeFormat Format, unsigned char Hours, unsigned char Minutes, unsigned char Seconds, unsigned char Frames, BMDTimecodeFlags Flags)
+		{
+			return WrappedVideoFrame->SetTimecodeFromComponents(Format, Hours, Minutes, Seconds, Frames, Flags);
+		}
+
+		HRESULT FDecklinkVideoFrame::SetAncillaryData(IDeckLinkVideoFrameAncillary* Ancillary)
+		{
+			return WrappedVideoFrame->SetAncillaryData(Ancillary);
+		}
+
+		HRESULT FDecklinkVideoFrame::SetTimecodeUserBits(BMDTimecodeFormat Format, BMDTimecodeUserBits UserBits)
+		{
+			return WrappedVideoFrame->SetTimecodeUserBits(Format, UserBits);
+		}
+
+		HRESULT FDecklinkVideoFrame::GetInt(BMDDeckLinkFrameMetadataID MetadataId, LONGLONG* Value)
+		{
+			HRESULT Result = S_OK;
+
+			switch (MetadataId)
+			{
+			case bmdDeckLinkFrameMetadataHDRElectroOpticalTransferFunc:
+			{
+				*Value = static_cast<int64>(HDRMetadata.EOTF);
+				break;
+			}
+			case bmdDeckLinkFrameMetadataColorspace:
+			{
+				switch (HDRMetadata.ColorSpace)
+				{
+				case EHDRMetaDataColorspace::Rec601:
+					*Value = bmdColorspaceRec601;
+					break;
+				case EHDRMetaDataColorspace::Rec709:
+					*Value = bmdColorspaceRec709;
+					break;
+				case EHDRMetaDataColorspace::Rec2020:
+					*Value = bmdColorspaceRec2020;
+					break;
+				default:
+					if (IsHDRLoggingOK())
+					{
+						UE_LOG(LogBlackmagicCore, Warning, TEXT("Invalid color space detected."));
+					}
+						
+					*Value = bmdColorspaceRec709;
+					break;
+				}
+				break;
+			}
+			default:
+				Result = E_INVALIDARG;
+				break;
+			}
+
+			return Result;
+		}
+
+		HRESULT FDecklinkVideoFrame::GetFloat(BMDDeckLinkFrameMetadataID MetadataId, double* Value)
+		{
+			HRESULT Result = S_OK;
+
+			switch (MetadataId)
+			{
+			case bmdDeckLinkFrameMetadataHDRDisplayPrimariesRedX:
+				*Value = HDRMetadata.DisplayPrimariesRedX;
+				break;
+			case bmdDeckLinkFrameMetadataHDRDisplayPrimariesRedY:
+				*Value = HDRMetadata.DisplayPrimariesRedY;
+				break;
+			case bmdDeckLinkFrameMetadataHDRDisplayPrimariesGreenX:
+				*Value = HDRMetadata.DisplayPrimariesGreenX;
+				break;
+			case bmdDeckLinkFrameMetadataHDRDisplayPrimariesGreenY:
+				*Value = HDRMetadata.DisplayPrimariesGreenY;
+				break;
+			case bmdDeckLinkFrameMetadataHDRDisplayPrimariesBlueX:
+				*Value = HDRMetadata.DisplayPrimariesBlueX;
+				break;
+			case bmdDeckLinkFrameMetadataHDRDisplayPrimariesBlueY:
+				*Value = HDRMetadata.DisplayPrimariesBlueY;
+				break;
+			case bmdDeckLinkFrameMetadataHDRWhitePointX:
+				*Value = HDRMetadata.WhitePointX;
+				break;
+			case bmdDeckLinkFrameMetadataHDRWhitePointY:
+                *Value = HDRMetadata.WhitePointY;
+                break;
+			case bmdDeckLinkFrameMetadataHDRMaxDisplayMasteringLuminance:
+				*Value = HDRMetadata.MaxDisplayLuminance;
+				break;
+			case bmdDeckLinkFrameMetadataHDRMinDisplayMasteringLuminance:
+				*Value = HDRMetadata.MinDisplayLuminance;
+				break;
+			case bmdDeckLinkFrameMetadataHDRMaximumContentLightLevel:
+				*Value = HDRMetadata.MaxContentLightLevel;
+				break;
+			case bmdDeckLinkFrameMetadataHDRMaximumFrameAverageLightLevel:
+				*Value = HDRMetadata.MaxFrameAverageLightLevel;
+				break;
+			default:
+				Result = E_INVALIDARG;
+				break;
+			}
+			return Result;
+		}
+
+		HRESULT FDecklinkVideoFrame::GetFlag(BMDDeckLinkFrameMetadataID MetadataId, BOOL* Value)
+		{
+			return E_INVALIDARG;
+		}
+
+		HRESULT FDecklinkVideoFrame::GetString(BMDDeckLinkFrameMetadataID MetadataId, BSTR* Value)
+		{
+			return E_INVALIDARG;
+		}
+
+		HRESULT FDecklinkVideoFrame::GetBytes(BMDDeckLinkFrameMetadataID MetadataId, void* Buffer, unsigned int* BufferSize)
+		{
+			return E_INVALIDARG;
+		}
+
+		bool FDecklinkVideoFrame::IsHDRLoggingOK()
+		{
+			// Reset log count after a few seconds.
+            if (FPlatformTime::Seconds() > LastHDRLogResetTime)
+            {
+            	HDRLogCount = 0;
+            }
+            
+            bool bIsOK = HDRLogCount < 30;
+
+            if (bIsOK)
+            {
+            	HDRLogCount++;
+            	LastHDRLogResetTime = FPlatformTime::Seconds() + 2;
+            }
+
+            return bIsOK;
+		}
+	}
 };
 
 CustomAllocator::CustomAllocator()
diff --git a/Engine/Plugins/Media/BlackmagicMedia/Source/BlackmagicCore/Private/BlackmagicOutputChannel.h b/Engine/Plugins/Media/BlackmagicMedia/Source/BlackmagicCore/Private/BlackmagicOutputChannel.h
--- a/Engine/Plugins/Media/BlackmagicMedia/Source/BlackmagicCore/Private/BlackmagicOutputChannel.h
+++ b/Engine/Plugins/Media/BlackmagicMedia/Source/BlackmagicCore/Private/BlackmagicOutputChannel.h
@@ -2,11 +2,12 @@
 
 #pragma once
 
+#include "Common.h"
+#include "GPUTextureTransferModule.h"
 #include "ThreadSafeQueue.h"
 #include "ThreadSafeObjectPool.h"
-#include <atomic>
 
-#include "GPUTextureTransferModule.h"
+#include <atomic>
 
 class CustomAllocator : public IDeckLinkMemoryAllocator
 {
@@ -109,6 +110,71 @@ namespace BlackmagicDesign
 			std::atomic_char32_t RefCount;
 		};
 
+		/**
+		 * Wrapper class around IDeckLinkMutableVideoFrame that also extends IDeckLinkVideoFrameMetadataExtensions in order to provide HDR metadata to the SDI feed.
+		 */
+		class FDecklinkVideoFrame : public IDeckLinkMutableVideoFrame, public IDeckLinkVideoFrameMetadataExtensions
+		{
+		public:
+			FDecklinkVideoFrame(IDeckLinkMutableVideoFrame* InWrappedVideoFrame, FHDRMetaData InHDRMetadata)
+				: RefCount(1)
+				, WrappedVideoFrame(InWrappedVideoFrame)
+				, HDRMetadata(MoveTemp(InHDRMetadata))
+			{
+			}
+
+			virtual ~FDecklinkVideoFrame() = default;
+
+			//~ Begin IUnknown interface
+			virtual HRESULT STDMETHODCALLTYPE QueryInterface(REFIID RIID, LPVOID* PPV) override;
+			virtual ULONG STDMETHODCALLTYPE AddRef() override;
+			virtual ULONG STDMETHODCALLTYPE Release() override;
+			//~ End IUnknown interface
+
+			//~ Begin IDeckLinkVideoFrame Interface
+			virtual BMDFrameFlags STDMETHODCALLTYPE GetFlags(void) override;
+			virtual long STDMETHODCALLTYPE GetWidth(void) override;
+			virtual long STDMETHODCALLTYPE GetHeight(void) override;
+			virtual long STDMETHODCALLTYPE GetRowBytes(void) override;
+			virtual BMDPixelFormat STDMETHODCALLTYPE GetPixelFormat(void) override;
+			virtual HRESULT STDMETHODCALLTYPE GetBytes(/* [out] */ void** Buffer) override;
+			virtual HRESULT STDMETHODCALLTYPE GetTimecode(/* [in] */ BMDTimecodeFormat Format, /* [out] */ IDeckLinkTimecode** Timecode) override;
+			virtual HRESULT STDMETHODCALLTYPE GetAncillaryData(/* [out] */ IDeckLinkVideoFrameAncillary** ancillary) override;
+			//~ End IDeckLinkVideoFrame Interface
+
+			//~ Begin IDeckLinkMutableVideoFrame Interface
+			virtual HRESULT STDMETHODCALLTYPE SetFlags(/* [in] */ BMDFrameFlags NewFlags) override;
+			virtual HRESULT STDMETHODCALLTYPE SetTimecode(/* [in] */ BMDTimecodeFormat Format, /* [in] */ IDeckLinkTimecode* Timecode) override;
+			virtual HRESULT STDMETHODCALLTYPE SetTimecodeFromComponents(/* [in] */ BMDTimecodeFormat Format, /* [in] */ unsigned char Hours, /* [in] */ unsigned char Minutes, /* [in] */ unsigned char Seconds, /* [in] */ unsigned char Frames, /* [in] */ BMDTimecodeFlags Flags) override;
+			virtual HRESULT STDMETHODCALLTYPE SetAncillaryData(/* [in] */ IDeckLinkVideoFrameAncillary* Ancillary) override; 
+			virtual HRESULT STDMETHODCALLTYPE SetTimecodeUserBits(/* [in] */ BMDTimecodeFormat Format, /* [in] */ BMDTimecodeUserBits UserBits) override;
+			//~ End IDeckLinkMutableVideoFrame Interface
+
+			//~ Begin IDeckLinkVideoFrameMetadataExtensions interface
+			virtual HRESULT STDMETHODCALLTYPE GetInt(BMDDeckLinkFrameMetadataID MetadataId, LONGLONG* Value) override;
+			virtual HRESULT STDMETHODCALLTYPE GetFloat(BMDDeckLinkFrameMetadataID MetadataId, double* Value) override;
+			virtual HRESULT STDMETHODCALLTYPE GetFlag(BMDDeckLinkFrameMetadataID MetadataId, BOOL* Value) override;
+			virtual HRESULT STDMETHODCALLTYPE GetString(BMDDeckLinkFrameMetadataID MetadataId, BSTR* Value) override;
+			virtual HRESULT STDMETHODCALLTYPE GetBytes(BMDDeckLinkFrameMetadataID MetadataId, void* Buffer, unsigned int* BufferSize) override;
+			//~ End IDeckLinkVideoFrameMetadataExtensions interface
+			
+		private:
+			/** Utility function to throttle HDR warning logs. */
+			bool IsHDRLoggingOK();
+
+		private:
+			/** Used to destroy object at 0 refs. */
+			ULONG RefCount;
+			/** Wrapped video frame. */
+			ReferencePtr<IDeckLinkMutableVideoFrame> WrappedVideoFrame;
+			/** Holds HDR Metadata. */
+			FHDRMetaData HDRMetadata;
+			/** Used for throttling HDR logs. */
+			int32_t HDRLogCount;
+			/** Used for keeping track of the last HDR logging reset time.  */
+			double LastHDRLogResetTime = 0.0;
+		};
+
 		class FOutputChannel
 		{
 		private:
@@ -242,7 +308,7 @@ namespace BlackmagicDesign
 
 				static const uint32_t InvalidFrameIdentifier;
 
-				IDeckLinkMutableVideoFrame* BlackmagicFrame;
+				ReferencePtr<FDecklinkVideoFrame> BlackmagicFrame;
 				uint8_t* AudioBuffer;
 				uint32_t NumAudioSamples;
 
diff --git a/Engine/Plugins/Media/BlackmagicMedia/Source/BlackmagicCore/Private/Platform/Linux/LinuxPlatform.cpp b/Engine/Plugins/Media/BlackmagicMedia/Source/BlackmagicCore/Private/Platform/Linux/LinuxPlatform.cpp
--- a/Engine/Plugins/Media/BlackmagicMedia/Source/BlackmagicCore/Private/Platform/Linux/LinuxPlatform.cpp
+++ b/Engine/Plugins/Media/BlackmagicMedia/Source/BlackmagicCore/Private/Platform/Linux/LinuxPlatform.cpp
@@ -12,6 +12,11 @@
 #include <pthread.h>
 #include <sys/mman.h>
 
+bool IsEqualIID(REFIID LHS, REFIID RHS)
+{
+	return memcmp(&LHS, &RHS, sizeof(REFIID)) == 0;
+}
+
 namespace BlackmagicPlatform
 {
 
diff --git a/Engine/Plugins/Media/BlackmagicMedia/Source/BlackmagicCore/Private/Platform/Linux/LinuxPlatform.h b/Engine/Plugins/Media/BlackmagicMedia/Source/BlackmagicCore/Private/Platform/Linux/LinuxPlatform.h
--- a/Engine/Plugins/Media/BlackmagicMedia/Source/BlackmagicCore/Private/Platform/Linux/LinuxPlatform.h
+++ b/Engine/Plugins/Media/BlackmagicMedia/Source/BlackmagicCore/Private/Platform/Linux/LinuxPlatform.h
@@ -8,6 +8,7 @@
 #include <locale>
 #include <thread>
 
+#include "LinuxCOM.h"
 
 // provide string the way Unreal Engine expects
 #define TCHAR char16_t
@@ -15,6 +16,10 @@
 // convert Windows style enums in Linux, i.e. BMDPixelFormat to _BMDPixelFormat
 #define ENUM(x) _ ## x
 
+#define BSTR const char*
+
+#define LONGLONG int64_t
+
 // simulate Windows BOOL type
 typedef bool BOOL;
 
@@ -23,6 +28,9 @@ class IDeckLinkDisplayMode;
 class IDeckLinkIterator;
 class IDeckLinkVideoConversion;
 
+/** Defined outside of namespace to match windows' implementation */
+bool IsEqualIID(REFIID LHS, REFIID RHS);
+
 namespace BlackmagicPlatform
 {
 	std::string GetName(IDeckLinkDisplayMode* DisplayMode);
diff --git a/Engine/Plugins/Media/BlackmagicMedia/Source/BlackmagicCore/Public/BlackmagicLib.h b/Engine/Plugins/Media/BlackmagicMedia/Source/BlackmagicCore/Public/BlackmagicLib.h
--- a/Engine/Plugins/Media/BlackmagicMedia/Source/BlackmagicCore/Public/BlackmagicLib.h
+++ b/Engine/Plugins/Media/BlackmagicMedia/Source/BlackmagicCore/Public/BlackmagicLib.h
@@ -178,6 +178,44 @@ namespace BlackmagicDesign
 		bool operator==(FChannelInfo& Other) const;
 	};
 
+
+	/** HDR Metadata. */
+	struct BLACKMAGICCORE_API FHDRMetaData
+	{
+		FHDRMetaData();
+
+		/** Whether the data contained in this struct is valid. */
+		bool bIsAvailable;
+		/** Target color space. */
+		EHDRMetaDataColorspace ColorSpace;
+		/** HDR Transfer function. */
+		EHDRMetaDataEOTF EOTF;
+		/** White point X coordinate. */
+		double WhitePointX;
+		/** White point Y coordinate. */
+		double WhitePointY;
+		/** Red chromaticity X coordinate. */
+		double DisplayPrimariesRedX;
+		/** Red chromaticity Y coordinate. */
+		double DisplayPrimariesRedY;
+		/** Green chromaticity X coordinate. */
+		double DisplayPrimariesGreenX;
+		/** Green chromaticity Y coordinate. */
+		double DisplayPrimariesGreenY;
+		/** Blue chromaticity X coordinate. */
+		double DisplayPrimariesBlueX;
+		/** Blue chromaticity Y coordinate. */
+		double DisplayPrimariesBlueY;
+		/** Max display mastering luminance. */
+		double MaxDisplayLuminance;
+		/** Min display mastering luminance. */
+		double MinDisplayLuminance;
+		/** Max content light level. */
+		double MaxContentLightLevel;
+		/** Max frame average light level. */
+		double MaxFrameAverageLightLevel;
+	};
+
 	/* FInputChannelOptions definition
 	*****************************************************************************/
 	struct BLACKMAGICCORE_API FInputChannelOptions
@@ -218,6 +256,9 @@ namespace BlackmagicDesign
 		ETimecodeFormat TimecodeFormat;
 		ELinkConfiguration LinkConfiguration;
 
+		/** HDR Metadata. */
+		FHDRMetaData HDRMetadata;
+
 		bool bOutputKey;
 		bool bOutputVideo;
 		bool bOutputAudio;
@@ -251,24 +292,6 @@ namespace BlackmagicDesign
 			EPixelFormat PixelFormat;
 			EFullPixelFormat FullPixelFormat;
 			EFieldDominance FieldDominance;
-
-			// HDR
-			struct BLACKMAGICCORE_API FHDRMetaData
-			{
-				FHDRMetaData();
-
-				bool bIsAvailable;
-				EHDRMetaDataColorspace ColorSpace;
-				EHDRMetaDataEOTF EOTF;
-				double WhitePointX;
-				double WhitePointY;
-				double DisplayPrimariesRedX;
-				double DisplayPrimariesRedY;
-				double DisplayPrimariesGreenX;
-				double DisplayPrimariesGreenY;
-				double DisplayPrimariesBlueX;
-				double DisplayPrimariesBlueY;
-			};
 			FHDRMetaData HDRMetaData;
 
 			// Audio
diff --git a/Engine/Plugins/Media/BlackmagicMedia/Source/BlackmagicMediaOutput/BlackmagicMediaOutput.Build.cs b/Engine/Plugins/Media/BlackmagicMedia/Source/BlackmagicMediaOutput/BlackmagicMediaOutput.Build.cs
--- a/Engine/Plugins/Media/BlackmagicMedia/Source/BlackmagicMediaOutput/BlackmagicMediaOutput.Build.cs
+++ b/Engine/Plugins/Media/BlackmagicMedia/Source/BlackmagicMediaOutput/BlackmagicMediaOutput.Build.cs
@@ -18,6 +18,7 @@ namespace UnrealBuildTool.Rules
 				new string[]
 				{
 					"BlackmagicCore",
+					"ColorManagement",
 					"Core",
 					"CoreUObject",
 					"Engine",
diff --git a/Engine/Plugins/Media/BlackmagicMedia/Source/BlackmagicMediaOutput/Private/BlackmagicMediaCapture.cpp b/Engine/Plugins/Media/BlackmagicMedia/Source/BlackmagicMediaOutput/Private/BlackmagicMediaCapture.cpp
--- a/Engine/Plugins/Media/BlackmagicMedia/Source/BlackmagicMediaOutput/Private/BlackmagicMediaCapture.cpp
+++ b/Engine/Plugins/Media/BlackmagicMedia/Source/BlackmagicMediaOutput/Private/BlackmagicMediaCapture.cpp
@@ -5,6 +5,7 @@
 #include "BlackmagicLib.h"
 #include "BlackmagicMediaOutput.h"
 #include "BlackmagicMediaOutputModule.h"
+#include "ColorSpace.h"
 #include "GPUTextureTransfer.h"
 #include "GPUTextureTransferModule.h"
 #include "Engine/Engine.h"
@@ -37,6 +38,63 @@ static FAutoConsoleCommand BlackmagicWriteInputRawDataCmd(
 
 namespace BlackmagicMediaCaptureHelpers
 {
+	BlackmagicDesign::FHDRMetaData MakeBlackmagicHDRMetadata(const FBlackmagicMediaHDROptions& HDROptions)
+	{
+		BlackmagicDesign::FHDRMetaData HDRMetadata;
+		HDRMetadata.bIsAvailable = true;
+		HDRMetadata.EOTF = (BlackmagicDesign::EHDRMetaDataEOTF) HDROptions.EOTF;
+
+		UE::Color::EColorSpace ColorSpaceEnum = UE::Color::EColorSpace::sRGB;
+		switch (HDROptions.Gamut)
+		{
+		case EBlackmagicHDRMetadataGamut::Rec709:
+			HDRMetadata.ColorSpace = BlackmagicDesign::EHDRMetaDataColorspace::Rec709;
+			ColorSpaceEnum = UE::Color::EColorSpace::sRGB;
+			break;
+		case EBlackmagicHDRMetadataGamut::Rec2020:
+			HDRMetadata.ColorSpace = BlackmagicDesign::EHDRMetaDataColorspace::Rec2020;
+			ColorSpaceEnum = UE::Color::EColorSpace::Rec2020;
+			break;
+		default:
+			HDRMetadata.bIsAvailable = false;
+			checkNoEntry();
+			break;
+		}
+		
+		UE::Color::FColorSpace ColorSpace(ColorSpaceEnum);
+		HDRMetadata.WhitePointX = ColorSpace.GetWhiteChromaticity().X;
+		HDRMetadata.WhitePointY = ColorSpace.GetWhiteChromaticity().Y;
+
+		HDRMetadata.DisplayPrimariesRedX = ColorSpace.GetRedChromaticity().X;
+		HDRMetadata.DisplayPrimariesRedY = ColorSpace.GetRedChromaticity().Y;
+
+		HDRMetadata.DisplayPrimariesGreenX = ColorSpace.GetGreenChromaticity().X;
+		HDRMetadata.DisplayPrimariesGreenY = ColorSpace.GetGreenChromaticity().Y;
+		
+		HDRMetadata.DisplayPrimariesBlueX = ColorSpace.GetBlueChromaticity().X;
+		HDRMetadata.DisplayPrimariesBlueY = ColorSpace.GetBlueChromaticity().Y;
+
+		return HDRMetadata;
+	}
+	
+	struct BLACKMAGICCORE_API FHDRMetaData
+	{
+		FHDRMetaData();
+
+		double WhitePointX;
+		double WhitePointY;
+		double DisplayPrimariesRedX;
+		double DisplayPrimariesRedY;
+		double DisplayPrimariesGreenX;
+		double DisplayPrimariesGreenY;
+		double DisplayPrimariesBlueX;
+		double DisplayPrimariesBlueY;
+		double MaxDisplayLuminance;
+		double MinDisplayLuminance;
+		double MaxContentLightLevel;
+		double MaxFrameAverageLightLevel;
+	};
+	
 	class FBlackmagicMediaCaptureEventCallback : public BlackmagicDesign::IOutputEventCallback
 	{
 	public:
@@ -553,6 +611,25 @@ bool UBlackmagicMediaCapture::HasFinishedProcessing() const
 	return Super::HasFinishedProcessing() || EventCallback == nullptr;
 }
 
+const FMatrix& UBlackmagicMediaCapture::GetRGBToYUVConversionMatrix() const
+{
+	if (const UBlackmagicMediaOutput* BMOutput = Cast<UBlackmagicMediaOutput>(MediaOutput))
+	{
+		switch(BMOutput->HDROptions.Gamut)
+		{
+		case EBlackmagicHDRMetadataGamut::Rec709:
+			return MediaShaders::RgbToYuvRec709Scaled;
+		case EBlackmagicHDRMetadataGamut::Rec2020:
+			return MediaShaders::RgbToYuvRec2020Scaled;
+		default:
+			checkNoEntry();
+			return MediaShaders::RgbToYuvRec709Scaled;
+		}
+	}
+		
+	return Super::GetRGBToYUVConversionMatrix();
+}
+
 bool UBlackmagicMediaCapture::InitBlackmagic(UBlackmagicMediaOutput* InBlackmagicMediaOutput)
 {
 	check(InBlackmagicMediaOutput);
@@ -605,6 +682,7 @@ bool UBlackmagicMediaCapture::InitBlackmagic(UBlackmagicMediaOutput* InBlackmagi
 	ChannelOptions.bLogDropFrames = bLogDropFrame;
 	ChannelOptions.bUseGPUDMA = ShouldCaptureRHIResource();
 	ChannelOptions.bScheduleInDifferentThread = InBlackmagicMediaOutput->bUseMultithreadedScheduling;
+	ChannelOptions.HDRMetadata = BlackmagicMediaCaptureHelpers::MakeBlackmagicHDRMetadata(InBlackmagicMediaOutput->HDROptions);
 
 	AudioBitDepth = InBlackmagicMediaOutput->AudioBitDepth;
 	bOutputAudio = InBlackmagicMediaOutput->bOutputAudio;
diff --git a/Engine/Plugins/Media/BlackmagicMedia/Source/BlackmagicMediaOutput/Public/BlackmagicMediaCapture.h b/Engine/Plugins/Media/BlackmagicMedia/Source/BlackmagicMediaOutput/Public/BlackmagicMediaCapture.h
--- a/Engine/Plugins/Media/BlackmagicMedia/Source/BlackmagicMediaOutput/Public/BlackmagicMediaCapture.h
+++ b/Engine/Plugins/Media/BlackmagicMedia/Source/BlackmagicMediaOutput/Public/BlackmagicMediaCapture.h
@@ -61,6 +61,7 @@ protected:
 	{
 		return true;
 	}
+	virtual const FMatrix& GetRGBToYUVConversionMatrix() const override;
 
 private:
 	struct FBlackmagicOutputCallback;
diff --git a/Engine/Plugins/Media/BlackmagicMedia/Source/BlackmagicMediaOutput/Public/BlackmagicMediaOutput.h b/Engine/Plugins/Media/BlackmagicMedia/Source/BlackmagicMediaOutput/Public/BlackmagicMediaOutput.h
--- a/Engine/Plugins/Media/BlackmagicMedia/Source/BlackmagicMediaOutput/Public/BlackmagicMediaOutput.h
+++ b/Engine/Plugins/Media/BlackmagicMedia/Source/BlackmagicMediaOutput/Public/BlackmagicMediaOutput.h
@@ -3,8 +3,6 @@
 #pragma once
 
 #include "MediaOutput.h"
-
-#include "BlackmagicDeviceProvider.h"
 #include "MediaIOCoreDefinitions.h"
 
 #include "BlackmagicMediaOutput.generated.h"
@@ -40,6 +38,46 @@ enum class EBlackmagicMediaOutputAudioBitDepth : uint8
 	Signed_32Bits = 32 UMETA(DisplayName = "32 bits signed")
 };
 
+/**
+ * HDR Transfer function.
+ * Must match BlackmagicDesign::EHDRMetaDataEOTF
+ */
+UENUM()
+enum class EBlackmagicHDRMetadataEOTF : uint8
+{
+	SDR,
+	HDR,
+	PQ,
+	HLG
+};
+	
+/**
+ * HDR Color Gamut.
+ */
+UENUM()
+enum class EBlackmagicHDRMetadataGamut : uint8
+{
+	Rec709,
+	Rec2020
+};
+
+
+/**
+ * Set of metadata describing a HDR video signal.
+ */
+USTRUCT()
+struct BLACKMAGICMEDIAOUTPUT_API FBlackmagicMediaHDROptions
+{
+	GENERATED_BODY()
+
+	/** Transfer function to use for converting the video signal to an optical signal. */
+	UPROPERTY(EditAnywhere, Category = "HDR")
+	EBlackmagicHDRMetadataEOTF EOTF = EBlackmagicHDRMetadataEOTF::SDR;
+
+	/** The color gamut of the video signal. */
+	UPROPERTY(EditAnywhere, Category = "HDR")
+	EBlackmagicHDRMetadataGamut Gamut = EBlackmagicHDRMetadataGamut::Rec709;
+};
 
 /**
  * Output information for a MediaCapture.
@@ -104,6 +142,12 @@ public:
 	 */
 	UPROPERTY(EditAnywhere, AdvancedDisplay, Category = "Output")
 	bool bInterlacedFieldsTimecodeNeedToMatch;
+	
+	/**
+	 * HDR Metadata of the video signal.
+	 */
+	UPROPERTY(EditAnywhere, Category = "Output", DisplayName="HDR")
+	FBlackmagicMediaHDROptions HDROptions;
 
 	/**
 	 * Whether to use multi threaded scheduling which should improve performance when outputting 4k and 8k content. (Experimental)
