Disable AnyThread capture with GPUDirect, since calling DVPMapBufferEnd from a different thread would crash

diff --git a/Engine/Plugins/Media/AjaMedia/Source/AjaMediaOutput/Private/AjaMediaCapture.cpp b/Engine/Plugins/Media/AjaMedia/Source/AjaMediaOutput/Private/AjaMediaCapture.cpp
--- a/Engine/Plugins/Media/AjaMedia/Source/AjaMediaOutput/Private/AjaMediaCapture.cpp
+++ b/Engine/Plugins/Media/AjaMedia/Source/AjaMediaOutput/Private/AjaMediaCapture.cpp
@@ -647,6 +647,12 @@ void UAjaMediaCapture::UnlockDMATexture_RenderThread(FTextureRHIRef InTexture)
 	}
 }
 
+bool UAjaMediaCapture::SupportsAnyThreadCapture() const
+{
+	// AnyThread is not supported with GPUDirect since calling dvpMapBufferEnd on a different thread will crash.
+	return !ShouldCaptureRHIResource();
+}
+
 void UAjaMediaCapture::OnFrameCaptured_AnyThread(const FCaptureBaseData& InBaseData, TSharedPtr<FMediaCaptureUserData, ESPMode::ThreadSafe> InUserData, const FMediaCaptureResourceData& InResourceData)
 {
 	TRACE_CPUPROFILER_EVENT_SCOPE(UAjaMediaCapture::OnFrameCaptured_AnyThread);
diff --git a/Engine/Plugins/Media/AjaMedia/Source/AjaMediaOutput/Public/AjaMediaCapture.h b/Engine/Plugins/Media/AjaMedia/Source/AjaMediaOutput/Public/AjaMediaCapture.h
--- a/Engine/Plugins/Media/AjaMedia/Source/AjaMediaOutput/Public/AjaMediaCapture.h
+++ b/Engine/Plugins/Media/AjaMedia/Source/AjaMediaOutput/Public/AjaMediaCapture.h
@@ -56,10 +56,7 @@ protected:
 	virtual void LockDMATexture_RenderThread(FTextureRHIRef InTexture) override;
 	virtual void UnlockDMATexture_RenderThread(FTextureRHIRef InTexture) override;
 	virtual void OnFrameCaptured_AnyThread(const FCaptureBaseData& InBaseData, TSharedPtr<FMediaCaptureUserData, ESPMode::ThreadSafe> InUserData, const FMediaCaptureResourceData& InResourceData) override;
-	virtual bool SupportsAnyThreadCapture() const override
-	{
-		return true;
-	}
+	virtual bool SupportsAnyThreadCapture() const override;
 	virtual const FMatrix& GetRGBToYUVConversionMatrix() const override;
 
 private:
diff --git a/Engine/Plugins/Media/BlackmagicMedia/Source/BlackmagicMediaOutput/Private/BlackmagicMediaCapture.cpp b/Engine/Plugins/Media/BlackmagicMedia/Source/BlackmagicMediaOutput/Private/BlackmagicMediaCapture.cpp
--- a/Engine/Plugins/Media/BlackmagicMedia/Source/BlackmagicMediaOutput/Private/BlackmagicMediaCapture.cpp
+++ b/Engine/Plugins/Media/BlackmagicMedia/Source/BlackmagicMediaOutput/Private/BlackmagicMediaCapture.cpp
@@ -799,6 +799,12 @@ void UBlackmagicMediaCapture::UnlockDMATexture_RenderThread(FTextureRHIRef InTex
 	}
 }
 
+bool UBlackmagicMediaCapture::SupportsAnyThreadCapture() const
+{
+	// AnyThread is not supported with GPUDirect since calling dvpMapBufferEnd on a different thread will crash.
+	return !ShouldCaptureRHIResource();
+}
+
 void UBlackmagicMediaCapture::OnFrameCaptured_RenderingThread(const FCaptureBaseData& InBaseData, TSharedPtr<FMediaCaptureUserData, ESPMode::ThreadSafe> InUserData, void* InBuffer, int32 Width, int32 Height, int32 BytesPerRow)
 {
 	TRACE_CPUPROFILER_EVENT_SCOPE(UBlackmagicMediaCapture::OnFrameCaptured_RenderingThread);
diff --git a/Engine/Plugins/Media/BlackmagicMedia/Source/BlackmagicMediaOutput/Public/BlackmagicMediaCapture.h b/Engine/Plugins/Media/BlackmagicMedia/Source/BlackmagicMediaOutput/Public/BlackmagicMediaCapture.h
--- a/Engine/Plugins/Media/BlackmagicMedia/Source/BlackmagicMediaOutput/Public/BlackmagicMediaCapture.h
+++ b/Engine/Plugins/Media/BlackmagicMedia/Source/BlackmagicMediaOutput/Public/BlackmagicMediaCapture.h
@@ -57,10 +57,7 @@ protected:
 	virtual void OnRHIResourceCaptured_RenderingThread(const FCaptureBaseData& InBaseData, TSharedPtr<FMediaCaptureUserData, ESPMode::ThreadSafe> InUserData, FTextureRHIRef InTexture) override;
 	virtual void LockDMATexture_RenderThread(FTextureRHIRef InTexture) override;
 	virtual void UnlockDMATexture_RenderThread(FTextureRHIRef InTexture) override;
-	virtual bool SupportsAnyThreadCapture() const override
-	{
-		return true;
-	}
+	virtual bool SupportsAnyThreadCapture() const override;
 	virtual const FMatrix& GetRGBToYUVConversionMatrix() const override;
 
 private:
diff --git a/Engine/Plugins/Media/MediaIOFramework/Source/MediaIOCore/Private/MediaCapture.cpp b/Engine/Plugins/Media/MediaIOFramework/Source/MediaIOCore/Private/MediaCapture.cpp
--- a/Engine/Plugins/Media/MediaIOFramework/Source/MediaIOCore/Private/MediaCapture.cpp
+++ b/Engine/Plugins/Media/MediaIOFramework/Source/MediaIOCore/Private/MediaCapture.cpp
@@ -346,7 +346,7 @@ bool UMediaCapture::StartSourceCapture(TSharedPtr<UE::MediaCapture::Private::FCa
 	}
 	else if (!SupportsAnyThreadCapture())
 	{
-		UE_LOG(LogMediaIOCore, Display, TEXT("AnyThread Capture was disabled because the media capture implementation does not have a AnyThread callback."));
+		UE_LOG(LogMediaIOCore, Display, TEXT("AnyThread Capture was disabled because the media capture implementation does not have a AnyThread callback or explicitly disabled AnyThread capture."));
 	}
 
 	return bInitialized;
