Fix timecode sometimes not working when applying media profiile on startup that overrides a timecode provider
Also fix issue with autodetectable timecode not reinitializing on video input format change
Fix Blackmagic video input being overriden by timecode provider.

diff --git a/Engine/Plugins/Media/BlackmagicMedia/Source/BlackmagicCore/Private/BlackmagicInputChannel.cpp b/Engine/Plugins/Media/BlackmagicMedia/Source/BlackmagicCore/Private/BlackmagicInputChannel.cpp
--- a/Engine/Plugins/Media/BlackmagicMedia/Source/BlackmagicCore/Private/BlackmagicInputChannel.cpp
+++ b/Engine/Plugins/Media/BlackmagicMedia/Source/BlackmagicCore/Private/BlackmagicInputChannel.cpp
@@ -178,7 +178,7 @@ namespace BlackmagicDesign
 				InputChannel->ChannelOptions.FormatInfo.FieldDominance = FormatInfo.FieldDominance;
 				InputChannel->ChannelOptions.FormatInfo.DisplayMode = InNewDisplayMode->GetDisplayMode();
 
-				if (InputChannel->ChannelOptions.bAutoDetect)
+				if (InputChannel->ChannelOptions.bAutoDetect || InputChannel->ChannelOptions.TimecodeFormat == ETimecodeFormat::TCF_Auto)
 				{
 					InputChannel->DisplayMode = InNewDisplayMode->GetDisplayMode();
 					InputChannel->Reinitialize();
@@ -535,7 +535,8 @@ namespace BlackmagicDesign
 
 		bool FInputChannel::IsCompatible(const FChannelInfo& InChannelInfo, const FInputChannelOptions& InChannelOptions) const
 		{
-			bool bResult = (ChannelOptions.FormatInfo.DisplayMode == InChannelOptions.FormatInfo.DisplayMode);
+			// Format isn't set on the channel options when using "Auto", since we don't care about the format, it's always going to be "compatible".
+			bool bResult = (ChannelOptions.FormatInfo.DisplayMode == InChannelOptions.FormatInfo.DisplayMode || InChannelOptions.bAutoDetect);
 
 			if (bResult)
 			{
diff --git a/Engine/Plugins/Media/BlackmagicMedia/Source/BlackmagicCore/Public/BlackmagicLib.h b/Engine/Plugins/Media/BlackmagicMedia/Source/BlackmagicCore/Public/BlackmagicLib.h
--- a/Engine/Plugins/Media/BlackmagicMedia/Source/BlackmagicCore/Public/BlackmagicLib.h
+++ b/Engine/Plugins/Media/BlackmagicMedia/Source/BlackmagicCore/Public/BlackmagicLib.h
@@ -146,15 +146,15 @@ namespace BlackmagicDesign
 	struct BLACKMAGICCORE_API FFormatInfo
 	{
 		/** Framerate */
-		uint32_t FrameRateNumerator;
-		uint32_t FrameRateDenominator;
+		uint32_t FrameRateNumerator = 0;
+		uint32_t FrameRateDenominator = 0;
 
 		/** Image Width & Weight in texels */
-		uint32_t Width;
-		uint32_t Height;
+		uint32_t Width = 0;
+		uint32_t Height = 0;
 
-		EFieldDominance FieldDominance;
-		FBlackmagicVideoFormat DisplayMode; // Unique identifier that represent all that combination for the device
+		EFieldDominance FieldDominance = EFieldDominance::Progressive;
+		FBlackmagicVideoFormat DisplayMode = 0; // Unique identifier that represent all that combination for the device
 
 		bool operator==(FFormatInfo& Other) const;
 
diff --git a/Engine/Plugins/Media/MediaFrameworkUtilities/Source/MediaFrameworkUtilities/Private/Profile/MediaProfile.cpp b/Engine/Plugins/Media/MediaFrameworkUtilities/Source/MediaFrameworkUtilities/Private/Profile/MediaProfile.cpp
--- a/Engine/Plugins/Media/MediaFrameworkUtilities/Source/MediaFrameworkUtilities/Private/Profile/MediaProfile.cpp
+++ b/Engine/Plugins/Media/MediaFrameworkUtilities/Source/MediaFrameworkUtilities/Private/Profile/MediaProfile.cpp
@@ -134,6 +134,14 @@ void UMediaProfile::Apply()
 		bTimecodeProvideWasApplied = true;
 		AppliedTimecodeProvider = TimecodeProvider;
 		PreviousTimecodeProvider = GEngine->GetTimecodeProvider();
+
+		if (PreviousTimecodeProvider && AppliedTimecodeProvider && PreviousTimecodeProvider->GetClass() == AppliedTimecodeProvider->GetClass())
+		{
+			GEngine->SetTimecodeProvider(nullptr);
+			// Sleep to let the timecode provider complete its shutdown on its hardware.
+			FPlatformProcess::Sleep(0.2);
+		}
+
 		bool bResult = GEngine->SetTimecodeProvider(TimecodeProvider);
 		if (!bResult && TimecodeProvider)
 		{
