Sort discovered blackmagic / aja devices by name

diff --git a/Engine/Plugins/Media/AjaMedia/Source/AjaMedia/Private/Aja/AjaDeviceProvider.cpp b/Engine/Plugins/Media/AjaMedia/Source/AjaMedia/Private/Aja/AjaDeviceProvider.cpp
--- a/Engine/Plugins/Media/AjaMedia/Source/AjaMedia/Private/Aja/AjaDeviceProvider.cpp
+++ b/Engine/Plugins/Media/AjaMedia/Source/AjaMedia/Private/Aja/AjaDeviceProvider.cpp
@@ -2,6 +2,7 @@
 
 #include "AjaDeviceProvider.h"
 
+#include "Algo/Sort.h"
 #include "Aja/Aja.h"
 #include "AjaMediaPrivate.h"
 
@@ -578,6 +579,8 @@ TArray<FMediaIOInputConfiguration> FAjaDeviceProvider::GetInputConfigurations()
 		}
 	}
 
+	Algo::SortBy(Results, [](const FMediaIOInputConfiguration& Config){ return Config.MediaConfiguration.MediaConnection.Device.DeviceName; }, FNameLexicalLess());
+
 	return Results;
 }
 
diff --git a/Engine/Plugins/Media/BlackmagicMedia/Source/BlackmagicMedia/Private/Blackmagic/BlackmagicDeviceProvider.cpp b/Engine/Plugins/Media/BlackmagicMedia/Source/BlackmagicMedia/Private/Blackmagic/BlackmagicDeviceProvider.cpp
--- a/Engine/Plugins/Media/BlackmagicMedia/Source/BlackmagicMedia/Private/Blackmagic/BlackmagicDeviceProvider.cpp
+++ b/Engine/Plugins/Media/BlackmagicMedia/Source/BlackmagicMedia/Private/Blackmagic/BlackmagicDeviceProvider.cpp
@@ -2,6 +2,7 @@
 
 #include "BlackmagicDeviceProvider.h"
 
+#include "Algo/Sort.h"
 #include "BlackmagicCoreModule.h"
 #include "BlackmagicLib.h"
 #include "BlackmagicMediaPrivate.h"
@@ -209,6 +210,8 @@ TArray<FMediaIOConfiguration> FBlackmagicDeviceProvider::GetConfigurations(bool
 		}
 	}
 
+	Algo::SortBy(Results, [](const FMediaIOConfiguration& Config){ return Config.MediaConnection.Device.DeviceName; }, FNameLexicalLess());
+
 	return Results;
 }
 
