Fix not being able to select a video format when using the blackmagic 8k pro for input.

diff --git a/Engine/Plugins/Media/BlackmagicMedia/Source/BlackmagicMedia/Private/Assets/BlackmagicMediaSource.cpp b/Engine/Plugins/Media/BlackmagicMedia/Source/BlackmagicMedia/Private/Assets/BlackmagicMediaSource.cpp
--- a/Engine/Plugins/Media/BlackmagicMedia/Source/BlackmagicMedia/Private/Assets/BlackmagicMediaSource.cpp
+++ b/Engine/Plugins/Media/BlackmagicMedia/Source/BlackmagicMedia/Private/Assets/BlackmagicMediaSource.cpp
@@ -19,15 +19,6 @@ UBlackmagicMediaSource::UBlackmagicMediaSource()
 	, bEncodeTimecodeInTexel(false)
 {
 	MediaConfiguration.bIsInput = true;
-	if (!MediaConfiguration.IsValid())
-	{
-		const FBlackmagicDeviceProvider DeviceProvider;
-		const TArray<FMediaIOConfiguration> Configurations = DeviceProvider.GetConfigurations();
-		if (Configurations.Num())
-		{
-			MediaConfiguration = Configurations[0];
-		}
-	}
 }
 
 /*
diff --git a/Engine/Plugins/Media/MediaIOFramework/Source/MediaIOEditor/Public/Customizations/MediaIOConfigurationCustomization.h b/Engine/Plugins/Media/MediaIOFramework/Source/MediaIOEditor/Public/Customizations/MediaIOConfigurationCustomization.h
--- a/Engine/Plugins/Media/MediaIOFramework/Source/MediaIOEditor/Public/Customizations/MediaIOConfigurationCustomization.h
+++ b/Engine/Plugins/Media/MediaIOFramework/Source/MediaIOEditor/Public/Customizations/MediaIOConfigurationCustomization.h
@@ -29,5 +29,5 @@ private:
 private:
 	TWeakPtr<SWidget> PermutationSelector;
 	FMediaIOConfiguration SelectedConfiguration;
-	bool bAutoDetectFormat = true;
+	bool bAutoDetectFormat = false;
 };
