Fix Blackmagic media player not working in shipping configuration.

This was caused by an essential function call being wrapped in a check() macro

diff --git a/Engine/Plugins/Media/BlackmagicMedia/Source/BlackmagicCore/Private/BlackmagicVideoFormats.cpp b/Engine/Plugins/Media/BlackmagicMedia/Source/BlackmagicCore/Private/BlackmagicVideoFormats.cpp
--- a/Engine/Plugins/Media/BlackmagicMedia/Source/BlackmagicCore/Private/BlackmagicVideoFormats.cpp
+++ b/Engine/Plugins/Media/BlackmagicMedia/Source/BlackmagicCore/Private/BlackmagicVideoFormats.cpp
@@ -119,7 +119,8 @@ namespace BlackmagicDesign
 
 			BMDTimeValue Scale = 0;
 			BMDTimeScale FrameDuration = 0;
-			check(InBlackmagicVideoMode->GetFrameRate(&Scale, &FrameDuration) == S_OK);
+			const HRESULT Result = InBlackmagicVideoMode->GetFrameRate(&Scale, &FrameDuration);
+			check(Result == S_OK);
 			NewFormat.FrameRateNumerator = (uint32_t)FrameDuration;
 			NewFormat.FrameRateDenominator = (uint32_t)Scale;
 
