CustomTimeStep - Add option to wait for both interlace fields before ticking the engine.

#rnx

diff --git a/Engine/Plugins/Media/AjaMedia/Source/AjaMedia/Private/Assets/AjaCustomTimeStep.cpp b/Engine/Plugins/Media/AjaMedia/Source/AjaMedia/Private/Assets/AjaCustomTimeStep.cpp
--- a/Engine/Plugins/Media/AjaMedia/Source/AjaMedia/Private/Assets/AjaCustomTimeStep.cpp
+++ b/Engine/Plugins/Media/AjaMedia/Source/AjaMedia/Private/Assets/AjaCustomTimeStep.cpp
@@ -393,7 +393,7 @@ FFrameRate UAjaCustomTimeStep::GetSyncRate() const
 {
 	FFrameRate SyncRate = GetFixedFrameRate();
 
-	if (MediaConfiguration.MediaMode.Standard == EMediaIOStandardType::ProgressiveSegmentedFrame)
+	if (MediaConfiguration.MediaMode.Standard == EMediaIOStandardType::ProgressiveSegmentedFrame || bWaitForBothFields)
 	{
 		// If pSF you should get 2 field interrupts.
 		SyncRate.Numerator *= 2;
diff --git a/Engine/Plugins/Media/BlackmagicMedia/Source/BlackmagicMedia/Private/Assets/BlackmagicCustomTimeStep.cpp b/Engine/Plugins/Media/BlackmagicMedia/Source/BlackmagicMedia/Private/Assets/BlackmagicCustomTimeStep.cpp
--- a/Engine/Plugins/Media/BlackmagicMedia/Source/BlackmagicMedia/Private/Assets/BlackmagicCustomTimeStep.cpp
+++ b/Engine/Plugins/Media/BlackmagicMedia/Source/BlackmagicMedia/Private/Assets/BlackmagicCustomTimeStep.cpp
@@ -20,7 +20,7 @@ namespace BlackmagicCustomTimeStepHelpers
 	class FInputEventCallback : public BlackmagicDesign::IInputEventCallback
 	{
 	public:
-		FInputEventCallback(UBlackmagicCustomTimeStep* InCustomTimeStep, const BlackmagicDesign::FChannelInfo& InChannelInfo, bool bInEnableOverrunDetection, uint32 InExpectedSyncDelta)
+		FInputEventCallback(UBlackmagicCustomTimeStep* InCustomTimeStep, const BlackmagicDesign::FChannelInfo& InChannelInfo, bool bInEnableOverrunDetection, uint32 InExpectedSyncDelta, bool bInWaitForBothFields)
 			: WeakOwner(InCustomTimeStep)
 			, RefCounter(0)
 			, ChannelInfo(InChannelInfo)
@@ -33,6 +33,7 @@ namespace BlackmagicCustomTimeStepHelpers
 			, PreviousSyncCount(0)
 			, SyncCountDelta(1)
 			, ExpectedSyncCountDelta(InExpectedSyncDelta)
+			, bWaitForBothFields(bInWaitForBothFields)
 		{
 		}
 
@@ -214,7 +215,7 @@ namespace BlackmagicCustomTimeStepHelpers
 		{
 			CurrentSyncCount = FrameNumber;
 
-			if (WaitSyncEvent)
+			if (WaitSyncEvent && !bWaitForBothFields)
 			{
 				WaitSyncEvent->Trigger();
 			}
@@ -242,6 +243,9 @@ namespace BlackmagicCustomTimeStepHelpers
 		uint64 PreviousSyncCount;
 		uint64 SyncCountDelta;
 		uint64 ExpectedSyncCountDelta;
+
+		/** Whether to wait for both fields before triggering the WaitForSync event. */
+		bool bWaitForBothFields = false;
 	};
 }
 
@@ -273,7 +277,7 @@ bool UBlackmagicCustomTimeStep::Initialize(class UEngine* InEngine)
 	BlackmagicDesign::FChannelInfo ChannelInfo;
 	ChannelInfo.DeviceIndex = MediaConfiguration.MediaConnection.Device.DeviceIdentifier;
 
-	InputEventCallback = new BlackmagicCustomTimeStepHelpers::FInputEventCallback(this, ChannelInfo, bEnableOverrunDetection, GetExpectedSyncCountDelta());
+	InputEventCallback = new BlackmagicCustomTimeStepHelpers::FInputEventCallback(this, ChannelInfo, bEnableOverrunDetection, GetExpectedSyncCountDelta(), bWaitForBothFields);
 
 	BlackmagicDesign::FInputChannelOptions ChannelOptions;
 	ChannelOptions.bAutoDetect = bAutoDetectFormat;
diff --git a/Engine/Source/Runtime/TimeManagement/Public/GenlockedCustomTimeStep.h b/Engine/Source/Runtime/TimeManagement/Public/GenlockedCustomTimeStep.h
--- a/Engine/Source/Runtime/TimeManagement/Public/GenlockedCustomTimeStep.h
+++ b/Engine/Source/Runtime/TimeManagement/Public/GenlockedCustomTimeStep.h
@@ -55,4 +55,10 @@ public:
 
 	/** Experimental fix for interlace field flipping issue. */
 	static TIMEMANAGEMENT_API TAutoConsoleVariable<int32> CVarExperimentalFieldFlipFix;
+
+	/** 
+	 * When using interlace video as a custom time step, enabling this option will force the engine to wait for both fields to be received before ticking. 
+	 */
+	UPROPERTY(EditAnywhere, Category = "Genlock Options")
+	bool bWaitForBothFields = false;
 };
