Fix freeze frame output when using blackmagic

diff --git a/Engine/Plugins/Media/BlackmagicMedia/Source/BlackmagicCore/Private/BlackmagicOutputChannel.cpp b/Engine/Plugins/Media/BlackmagicMedia/Source/BlackmagicCore/Private/BlackmagicOutputChannel.cpp
--- a/Engine/Plugins/Media/BlackmagicMedia/Source/BlackmagicCore/Private/BlackmagicOutputChannel.cpp
+++ b/Engine/Plugins/Media/BlackmagicMedia/Source/BlackmagicCore/Private/BlackmagicOutputChannel.cpp
@@ -1624,7 +1624,7 @@ namespace BlackmagicDesign
 				return E_INVALIDARG;
 			}
 
-			if (IsEqualIID(RIID, IID_IUnknown))
+			if (IsEqualIID(RIID, IID_IUnknown) || IsEqualIID(RIID, IID_IDeckLinkVideoFrame))
 			{
 				*PPV = static_cast<IDeckLinkMutableVideoFrame*>(this);
 				AddRef();
