Add GPUDirect calls to aja media capture and blackmagic media capture

#rnx

diff --git a/Engine/Plugins/Media/AjaMedia/Source/Aja/Private/OutputChannel.cpp b/Engine/Plugins/Media/AjaMedia/Source/Aja/Private/OutputChannel.cpp
--- a/Engine/Plugins/Media/AjaMedia/Source/Aja/Private/OutputChannel.cpp
+++ b/Engine/Plugins/Media/AjaMedia/Source/Aja/Private/OutputChannel.cpp
@@ -1126,6 +1126,11 @@ namespace AJA
 			uint32_t CurrentOutFrame = 0;
 			CurrentOutFrame ^= 1;
 
+			if (TextureTransfer)
+			{
+				TextureTransfer->ThreadPrep();
+			}
+
 			bool bRunning = true;
 			while (bRunning && !bStopRequested)
 			{
@@ -1294,6 +1299,11 @@ namespace AJA
 					GetOptions().CallbackInterface->OnCompletion(bRunning);
 				}
 			}
+
+			if (TextureTransfer)
+			{
+				TextureTransfer->ThreadCleanup();
+			}
 		}
 
 		void OutputChannelThread::Thread_TestInterlacedOutput(Frame* InFrame)
diff --git a/Engine/Plugins/Media/AjaMedia/Source/AjaMediaOutput/Private/AjaMediaCapture.cpp b/Engine/Plugins/Media/AjaMedia/Source/AjaMediaOutput/Private/AjaMediaCapture.cpp
--- a/Engine/Plugins/Media/AjaMedia/Source/AjaMediaOutput/Private/AjaMediaCapture.cpp
+++ b/Engine/Plugins/Media/AjaMedia/Source/AjaMediaOutput/Private/AjaMediaCapture.cpp
@@ -466,6 +466,14 @@ const FMatrix& UAjaMediaCapture::GetRGBToYUVConversionMatrix() const
 	return Super::GetRGBToYUVConversionMatrix();
 }
 
+void UAjaMediaCapture::WaitForGPU(FRHITexture* InRHITexture)
+{
+	if (TextureTransfer)
+	{
+		TextureTransfer->WaitForGPU(InRHITexture);
+	}
+}
+
 bool UAjaMediaCapture::InitAJA(UAjaMediaOutput* InAjaMediaOutput)
 {
 	check(InAjaMediaOutput);
@@ -668,8 +676,7 @@ void UAjaMediaCapture::UnlockDMATexture_RenderThread(FTextureRHIRef InTexture)
 
 bool UAjaMediaCapture::SupportsAnyThreadCapture() const
 {
-	// AnyThread is not supported with GPUDirect since calling dvpMapBufferEnd on a different thread will crash.
-	return !ShouldCaptureRHIResource();
+	return true;
 }
 
 void UAjaMediaCapture::OnFrameCaptured_AnyThread(const FCaptureBaseData& InBaseData, TSharedPtr<FMediaCaptureUserData, ESPMode::ThreadSafe> InUserData, const FMediaCaptureResourceData& InResourceData)
diff --git a/Engine/Plugins/Media/AjaMedia/Source/AjaMediaOutput/Public/AjaMediaCapture.h b/Engine/Plugins/Media/AjaMedia/Source/AjaMediaOutput/Public/AjaMediaCapture.h
--- a/Engine/Plugins/Media/AjaMedia/Source/AjaMediaOutput/Public/AjaMediaCapture.h
+++ b/Engine/Plugins/Media/AjaMedia/Source/AjaMediaOutput/Public/AjaMediaCapture.h
@@ -58,7 +58,7 @@ protected:
 	virtual void OnFrameCaptured_AnyThread(const FCaptureBaseData& InBaseData, TSharedPtr<FMediaCaptureUserData, ESPMode::ThreadSafe> InUserData, const FMediaCaptureResourceData& InResourceData) override;
 	virtual bool SupportsAnyThreadCapture() const override;
 	virtual const FMatrix& GetRGBToYUVConversionMatrix() const override;
-
+	virtual void WaitForGPU(FRHITexture* InRHITexture) override;
 private:
 	struct FAjaOutputCallback;
 	friend FAjaOutputCallback;
diff --git a/Engine/Plugins/Media/BlackmagicMedia/Source/BlackmagicCore/Private/BlackmagicOutputChannel.cpp b/Engine/Plugins/Media/BlackmagicMedia/Source/BlackmagicCore/Private/BlackmagicOutputChannel.cpp
--- a/Engine/Plugins/Media/BlackmagicMedia/Source/BlackmagicCore/Private/BlackmagicOutputChannel.cpp
+++ b/Engine/Plugins/Media/BlackmagicMedia/Source/BlackmagicCore/Private/BlackmagicOutputChannel.cpp
@@ -265,7 +265,10 @@ namespace BlackmagicDesign
 				return false;
 			}
 
-			TextureTransfer = FGPUTextureTransferModule::Get().GetTextureTransfer();
+			if (FGPUTextureTransferModule::Get().IsEnabled() && FGPUTextureTransferModule::Get().IsInitialized())
+			{
+				TextureTransfer = FGPUTextureTransferModule::Get().GetTextureTransfer();
+			}
 
 			if (ChannelOptions.bOutputAudio)
 			{
@@ -645,6 +648,10 @@ namespace BlackmagicDesign
 
 		void FOutputChannel::VideoThread_ScheduleVideo()
 		{
+			if (TextureTransfer)
+			{
+				TextureTransfer->ThreadPrep();
+			}
 			while (true)
 			{
 				if (!bScheduleVideoActive)
@@ -675,6 +682,11 @@ namespace BlackmagicDesign
 					}
 				}
 			}
+
+			if (TextureTransfer)
+			{
+				TextureTransfer->ThreadCleanup();
+			}
 		}
 
 		void FOutputChannel::AudioThread_ScheduleAudioSamples()
@@ -1332,6 +1344,7 @@ namespace BlackmagicDesign
 			
 			if (!bRegisteredDMABuffers)
 			{
+				TRACE_CPUPROFILER_EVENT_SCOPE(BlackmagicOutputChannel::RegisterGPUDirectBuffer);
 				for (FOutputFrame* Frame : AllFrames)
 				{
 					void* Buffer = nullptr;
@@ -1354,6 +1367,7 @@ namespace BlackmagicDesign
 			{
 				UE::GPUTextureTransfer::ETransferDirection Direction = UE::GPUTextureTransfer::ETransferDirection::GPU_TO_CPU;
 				
+				TRACE_CPUPROFILER_EVENT_SCOPE(BlackmagicOutputChannel::TransferTexture);
 				TextureTransfer->TransferTexture(DestinationBuffer, InFrame.RHITexture, Direction);
 				AvailableWritingFrame->bVideoLineFilled = true;
 				AvailableWritingFrame->bVideoF2LineFilled = true;
diff --git a/Engine/Plugins/Media/BlackmagicMedia/Source/BlackmagicMediaOutput/Private/BlackmagicMediaCapture.cpp b/Engine/Plugins/Media/BlackmagicMedia/Source/BlackmagicMediaOutput/Private/BlackmagicMediaCapture.cpp
--- a/Engine/Plugins/Media/BlackmagicMedia/Source/BlackmagicMediaOutput/Private/BlackmagicMediaCapture.cpp
+++ b/Engine/Plugins/Media/BlackmagicMedia/Source/BlackmagicMediaOutput/Private/BlackmagicMediaCapture.cpp
@@ -635,6 +635,14 @@ const FMatrix& UBlackmagicMediaCapture::GetRGBToYUVConversionMatrix() const
 	return Super::GetRGBToYUVConversionMatrix();
 }
 
+void UBlackmagicMediaCapture::WaitForGPU(FRHITexture* InRHITexture)
+{
+	if (TextureTransfer)
+	{
+		TextureTransfer->WaitForGPU(InRHITexture);
+	}
+}
+
 bool UBlackmagicMediaCapture::InitBlackmagic(UBlackmagicMediaOutput* InBlackmagicMediaOutput)
 {
 	check(InBlackmagicMediaOutput);
@@ -807,8 +815,7 @@ void UBlackmagicMediaCapture::UnlockDMATexture_RenderThread(FTextureRHIRef InTex
 
 bool UBlackmagicMediaCapture::SupportsAnyThreadCapture() const
 {
-	// AnyThread is not supported with GPUDirect since calling dvpMapBufferEnd on a different thread will crash.
-	return !ShouldCaptureRHIResource();
+	return true;
 }
 
 void UBlackmagicMediaCapture::OnFrameCaptured_RenderingThread(const FCaptureBaseData& InBaseData, TSharedPtr<FMediaCaptureUserData, ESPMode::ThreadSafe> InUserData, void* InBuffer, int32 Width, int32 Height, int32 BytesPerRow)
diff --git a/Engine/Plugins/Media/BlackmagicMedia/Source/BlackmagicMediaOutput/Public/BlackmagicMediaCapture.h b/Engine/Plugins/Media/BlackmagicMedia/Source/BlackmagicMediaOutput/Public/BlackmagicMediaCapture.h
--- a/Engine/Plugins/Media/BlackmagicMedia/Source/BlackmagicMediaOutput/Public/BlackmagicMediaCapture.h
+++ b/Engine/Plugins/Media/BlackmagicMedia/Source/BlackmagicMediaOutput/Public/BlackmagicMediaCapture.h
@@ -59,6 +59,8 @@ protected:
 	virtual void UnlockDMATexture_RenderThread(FTextureRHIRef InTexture) override;
 	virtual bool SupportsAnyThreadCapture() const override;
 	virtual const FMatrix& GetRGBToYUVConversionMatrix() const override;
+	virtual void WaitForGPU(FRHITexture* InRHITexture) override;
+
 
 private:
 	struct FBlackmagicOutputCallback;
