Fix some blackmagic media sources being in a bad state and being unable to set their configuration

This change ensures that invalid Aja and BM media sources get a correct configuration if the current one is invalid. The logic used to be incorrect where it would assign the first found configuration, which could be an output in some case, and would render the source unusable.

diff --git a/Engine/Plugins/Media/AjaMedia/Source/AjaMedia/Private/Assets/AjaMediaSource.cpp b/Engine/Plugins/Media/AjaMedia/Source/AjaMedia/Private/Assets/AjaMediaSource.cpp
--- a/Engine/Plugins/Media/AjaMedia/Source/AjaMedia/Private/Assets/AjaMediaSource.cpp
+++ b/Engine/Plugins/Media/AjaMedia/Source/AjaMedia/Private/Assets/AjaMediaSource.cpp
@@ -8,7 +8,8 @@
 #include "MediaIOCorePlayerBase.h"
 
 UAjaMediaSource::UAjaMediaSource()
-	: bCaptureWithAutoCirculating(true)
+	: UCaptureCardMediaSource()
+	, bCaptureWithAutoCirculating(true)
 	, bCaptureAncillary(false)
 	, MaxNumAncillaryFrameBuffer(8)
 	, bCaptureAudio(false)
@@ -20,13 +21,7 @@ UAjaMediaSource::UAjaMediaSource()
 	, bLogDropFrame(true)
 	, bEncodeTimecodeInTexel(false)
 {
-	MediaConfiguration.bIsInput = true;
-	const FAjaDeviceProvider DeviceProvider;
-	const TArray<FMediaIOConfiguration> Configurations = DeviceProvider.GetConfigurations();
-	if (Configurations.Num())
-	{
-		MediaConfiguration = Configurations[0];
-	}
+	AssignDefaultConfiguration();
 }
 
 /*
@@ -343,6 +338,8 @@ void UAjaMediaSource::PostLoad()
 {
 	Super::PostLoad();
 
+	AssignDefaultConfiguration();
+
 #if WITH_EDITORONLY_DATA
 	PRAGMA_DISABLE_DEPRECATION_WARNINGS
 	if (TimecodeFormat_DEPRECATED != EMediaIOTimecodeFormat::None)
@@ -363,3 +360,22 @@ void UAjaMediaSource::PostLoad()
 	PRAGMA_ENABLE_DEPRECATION_WARNINGS
 #endif
 }
+
+
+void UAjaMediaSource::AssignDefaultConfiguration()
+{
+	if (!MediaConfiguration.IsValid())
+	{
+		const FAjaDeviceProvider DeviceProvider;
+		const TArray<FMediaIOConfiguration> Configurations = DeviceProvider.GetConfigurations();
+		for (const FMediaIOConfiguration& Configuration : Configurations)
+		{
+			if (Configuration.bIsInput)
+			{
+				MediaConfiguration = Configuration;
+				bRenderJIT = false;
+				break;
+			}
+		}
+	}
+}
diff --git a/Engine/Plugins/Media/AjaMedia/Source/AjaMedia/Public/AjaMediaSource.h b/Engine/Plugins/Media/AjaMedia/Source/AjaMedia/Public/AjaMediaSource.h
--- a/Engine/Plugins/Media/AjaMedia/Source/AjaMedia/Public/AjaMediaSource.h
+++ b/Engine/Plugins/Media/AjaMedia/Source/AjaMedia/Public/AjaMediaSource.h
@@ -138,6 +138,7 @@ public:
 	virtual FString GetUrl() const override;
 	virtual bool Validate() const override;
 
+
 public:
 	//~ UObject interface
 #if WITH_EDITOR
@@ -148,4 +149,8 @@ public:
 	virtual void PostLoad() override;
 	virtual bool SupportsFormatAutoDetection() const override { return true; }
 	//~ End UObject interface
+
+private:
+	/** Assigns a default configuration if the current configuration is invalid. */
+	void AssignDefaultConfiguration();
 };
diff --git a/Engine/Plugins/Media/BlackmagicMedia/Source/BlackmagicMedia/Private/Assets/BlackmagicMediaSource.cpp b/Engine/Plugins/Media/BlackmagicMedia/Source/BlackmagicMedia/Private/Assets/BlackmagicMediaSource.cpp
--- a/Engine/Plugins/Media/BlackmagicMedia/Source/BlackmagicMedia/Private/Assets/BlackmagicMediaSource.cpp
+++ b/Engine/Plugins/Media/BlackmagicMedia/Source/BlackmagicMedia/Private/Assets/BlackmagicMediaSource.cpp
@@ -9,7 +9,8 @@
 #include "MediaIOCorePlayerBase.h"
 
 UBlackmagicMediaSource::UBlackmagicMediaSource()
-	: bCaptureAudio(false)
+	: UCaptureCardMediaSource()
+	, bCaptureAudio(false)
 	, AudioChannels(EBlackmagicMediaAudioChannel::Stereo2)
 	, MaxNumAudioFrameBuffer(8)
 	, bCaptureVideo(true)
@@ -18,7 +19,7 @@ UBlackmagicMediaSource::UBlackmagicMediaSource()
 	, bLogDropFrame(false)
 	, bEncodeTimecodeInTexel(false)
 {
-	MediaConfiguration.bIsInput = true;
+	AssignDefaultConfiguration();
 }
 
 /*
@@ -172,6 +173,22 @@ bool UBlackmagicMediaSource::Validate() const
 	return true;
 }
 
+void UBlackmagicMediaSource::AssignDefaultConfiguration()
+{
+	const FBlackmagicDeviceProvider DeviceProvider;
+	const TArray<FMediaIOConfiguration> Configurations = DeviceProvider.GetConfigurations();
+
+	for (const FMediaIOConfiguration& Configuration : Configurations)
+	{
+		if (Configuration.bIsInput)
+		{
+			MediaConfiguration = Configuration;
+			bRenderJIT = false;
+			break;
+		}
+	}
+}
+
 #if WITH_EDITOR
 bool UBlackmagicMediaSource::CanEditChange(const FProperty* InProperty) const
 {
@@ -228,6 +245,9 @@ void UBlackmagicMediaSource::PostEditChangeChainProperty(struct FPropertyChanged
 void UBlackmagicMediaSource::PostLoad()
 {
 	Super::PostLoad();
+
+	AssignDefaultConfiguration();
+
 #if WITH_EDITORONLY_DATA
 	PRAGMA_DISABLE_DEPRECATION_WARNINGS
 		if (TimecodeFormat_DEPRECATED != EMediaIOTimecodeFormat::None)
diff --git a/Engine/Plugins/Media/BlackmagicMedia/Source/BlackmagicMedia/Public/BlackmagicMediaSource.h b/Engine/Plugins/Media/BlackmagicMedia/Source/BlackmagicMedia/Public/BlackmagicMediaSource.h
--- a/Engine/Plugins/Media/BlackmagicMedia/Source/BlackmagicMedia/Public/BlackmagicMediaSource.h
+++ b/Engine/Plugins/Media/BlackmagicMedia/Source/BlackmagicMedia/Public/BlackmagicMediaSource.h
@@ -132,4 +132,8 @@ public:
 #endif //WITH_EDITOR
 	virtual void PostLoad() override;
 	//~ End UObject interface
+
+private:
+	/** Assigns a default configuration if the current configuration is invalid. */
+	void AssignDefaultConfiguration();
 };
