Blackmagic - Add an option to discard interlace frames if they land on an odd frame.

Add support for burning timecode on interlace frames, by burning the timecode of the two engine frames that were used to form the output frame.

diff --git a/Engine/Plugins/Media/BlackmagicMedia/Source/BlackmagicCore/BlackmagicCore.build.cs b/Engine/Plugins/Media/BlackmagicMedia/Source/BlackmagicCore/BlackmagicCore.build.cs
--- a/Engine/Plugins/Media/BlackmagicMedia/Source/BlackmagicCore/BlackmagicCore.build.cs
+++ b/Engine/Plugins/Media/BlackmagicMedia/Source/BlackmagicCore/BlackmagicCore.build.cs
@@ -10,6 +10,7 @@ public class BlackmagicCore : ModuleRules
 			new string[]
 			{
 				"BlackmagicSDK",
+				"MediaIOCore",
 				"GPUTextureTransfer"
 			});
 			
diff --git a/Engine/Plugins/Media/BlackmagicMedia/Source/BlackmagicCore/Private/BlackmagicDevice.cpp b/Engine/Plugins/Media/BlackmagicMedia/Source/BlackmagicCore/Private/BlackmagicDevice.cpp
--- a/Engine/Plugins/Media/BlackmagicMedia/Source/BlackmagicCore/Private/BlackmagicDevice.cpp
+++ b/Engine/Plugins/Media/BlackmagicMedia/Source/BlackmagicCore/Private/BlackmagicDevice.cpp
@@ -418,6 +418,15 @@ namespace BlackmagicDesign
 			return false;
 		}
 
+		bool FDevice::ShouldCaptureThisFrame(const FChannelInfo& InChannelInfo, ::FTimecode Timecode, uint32 FrameNumber)
+		{
+			if (FOutputChannel* FoundChannel = FindChannel(InChannelInfo))
+			{
+				return FoundChannel->ShouldCaptureThisFrame(Timecode, FrameNumber);
+			}
+			return false;
+		}
+
 		bool FDevice::SendVideoFrameData(const FChannelInfo& InChannelInfo, FFrameDescriptor_GPUDMA& InFrame)
 		{
 			// SendVideoFrameData and FUnregisterOutputCommand can't happen at the same time.
diff --git a/Engine/Plugins/Media/BlackmagicMedia/Source/BlackmagicCore/Private/BlackmagicDevice.h b/Engine/Plugins/Media/BlackmagicMedia/Source/BlackmagicCore/Private/BlackmagicDevice.h
--- a/Engine/Plugins/Media/BlackmagicMedia/Source/BlackmagicCore/Private/BlackmagicDevice.h
+++ b/Engine/Plugins/Media/BlackmagicMedia/Source/BlackmagicCore/Private/BlackmagicDevice.h
@@ -25,6 +25,7 @@ namespace BlackmagicDesign
 			FUniqueIdentifier RegisterOutputChannel(const FChannelInfo& InChannelInfo, const FOutputChannelOptions& InChannelOptions, ReferencePtr<IOutputEventCallback> InCallback);
 			void UnregisterOutputChannel(const FChannelInfo& InChannelInfo, FUniqueIdentifier InIdentifier);
 			bool SendVideoFrameData(const FChannelInfo& InChannelInfo, const FFrameDescriptor& InFrame);
+			bool ShouldCaptureThisFrame(const FChannelInfo& InChannelInfo, ::FTimecode Timecode, uint32 FrameNumber);
 			bool SendVideoFrameData(const FChannelInfo& InChannelInfo, FFrameDescriptor_GPUDMA& InFrame);
 			bool SendAudioSamples(const FChannelInfo& InChannelInfo, const FAudioSamplesDescriptor& InSamples);
 		public:
diff --git a/Engine/Plugins/Media/BlackmagicMedia/Source/BlackmagicCore/Private/BlackmagicLib.cpp b/Engine/Plugins/Media/BlackmagicMedia/Source/BlackmagicCore/Private/BlackmagicLib.cpp
--- a/Engine/Plugins/Media/BlackmagicMedia/Source/BlackmagicCore/Private/BlackmagicLib.cpp
+++ b/Engine/Plugins/Media/BlackmagicMedia/Source/BlackmagicCore/Private/BlackmagicLib.cpp
@@ -111,9 +111,11 @@ namespace BlackmagicDesign
 		, bOutputAudio(false)
 		, bOutputInterlacedFieldsTimecodeNeedToMatch(false)
 		, bOutputInterlaceAsProgressive(false)
+		, bOutputInterlaceOnEvenFrames(false)
 		, bLogDropFrames(false)
 		, bUseGPUDMA(false)
 		, bScheduleInDifferentThread(false)
+		, bEncodeTimecodeInTexel(false)
 	{
 	}
 
@@ -292,6 +294,10 @@ namespace BlackmagicDesign
 #endif
 	}
 
+	bool ShouldCaptureThisFrame(const FChannelInfo& InChannelInfo, ::FTimecode Timecode, uint32 FrameNumber)
+	{
+		return Private::FDevice::GetDevice()->ShouldCaptureThisFrame(InChannelInfo, Timecode, FrameNumber);
+	}
 
 	bool SendAudioSamples(const FChannelInfo& InChannelInfo, const FAudioSamplesDescriptor& InSamples)
 	{
diff --git a/Engine/Plugins/Media/BlackmagicMedia/Source/BlackmagicCore/Private/BlackmagicOutputChannel.cpp b/Engine/Plugins/Media/BlackmagicMedia/Source/BlackmagicCore/Private/BlackmagicOutputChannel.cpp
--- a/Engine/Plugins/Media/BlackmagicMedia/Source/BlackmagicCore/Private/BlackmagicOutputChannel.cpp
+++ b/Engine/Plugins/Media/BlackmagicMedia/Source/BlackmagicCore/Private/BlackmagicOutputChannel.cpp
@@ -5,9 +5,12 @@
 #include "Common.h"
 
 #include "BlackmagicHelper.h"
+#include "MediaIOCoreEncodeTime.h"
+#include "SSE2MemCpy.h"
+
 #include <iostream>
 #include <fstream>
-#include "SSE2MemCpy.h"
+
 
 #if PLATFORM_WINDOWS
 #include "GPUTextureTransfer.h"
@@ -105,6 +108,31 @@ namespace BlackmagicDesign
 			return RefCount;
 		}
 
+		void EncodeTimecodeInTexel(BMDPixelFormat PixelFormat, BlackmagicDesign::FTimecode Timecode, void* Buffer, int32 Width, int32 Height)
+		{
+			switch (PixelFormat)
+			{
+			case bmdFormat8BitYUV:
+			{
+				FMediaIOCoreEncodeTime EncodeTime(EMediaIOCoreEncodePixelFormat::CharUYVY, Buffer, Width * 4, Width * 2, Height);
+				EncodeTime.Render(Timecode.Hours, Timecode.Minutes, Timecode.Seconds, Timecode.Frames);
+				break;
+			}
+			case bmdFormat8BitBGRA:
+			{
+				FMediaIOCoreEncodeTime EncodeTime(EMediaIOCoreEncodePixelFormat::CharBGRA, Buffer, Width * 4, Width, Height);
+				EncodeTime.Render(Timecode.Hours, Timecode.Minutes, Timecode.Seconds, Timecode.Frames);
+				break;
+			}
+			case bmdFormat10BitYUV:
+			{
+				FMediaIOCoreEncodeTime EncodeTime(EMediaIOCoreEncodePixelFormat::YUVv210, Buffer, Width * 16, Width * 6, Height);
+				EncodeTime.Render(Timecode.Hours, Timecode.Minutes, Timecode.Seconds, Timecode.Frames);
+				break;
+			}
+			}
+		}
+
 		/* Frame implementation
 		*****************************************************************************/
 		const uint32_t FOutputChannel::FOutputFrame::InvalidFrameIdentifier = static_cast<uint32_t>(-1);
@@ -601,6 +629,38 @@ namespace BlackmagicDesign
 				EmbedTimecodeInFrame(InFrame);
 			}
 
+			void* DestinationBuffer;
+			InFrame->BlackmagicFrame->GetBytes(&DestinationBuffer);
+
+			const int32_t VideoHeight = ChannelOptions.FormatInfo.Height;
+			const uint32_t Stride = InFrame->BlackmagicFrame->GetRowBytes();
+			const uint32_t VideoWidth = ChannelOptions.PixelFormat == EPixelFormat::pf_8Bits ? Stride / 4 : Stride / 16;
+			
+			if (ChannelOptions.FormatInfo.FieldDominance == EFieldDominance::Progressive)
+			{
+				if (ChannelOptions.bEncodeTimecodeInTexel)
+				{
+					EncodeTimecodeInTexel(PixelFormat, InFrame->Timecode, DestinationBuffer, VideoWidth, VideoHeight);
+				}
+			}
+			else
+			{
+				if (ChannelOptions.bEncodeTimecodeInTexel)
+				{
+					EncodeTimecodeInTexel(PixelFormat, InFrame->Timecode, DestinationBuffer, VideoWidth, VideoHeight);
+					EncodeTimecodeInTexel(PixelFormat, InFrame->Timecode, (void*)((uint8*)DestinationBuffer + 200 * 4), VideoWidth, VideoHeight);
+				}
+
+				if (InFrame->Timecode.Frames % 2 == 0)
+				{
+					UE_LOG(LogBlackmagicCore, VeryVerbose, TEXT("Scheduled frame starting with even timecode. (%d)"), InFrame->Timecode.Frames);
+				}
+				else
+				{
+					UE_LOG(LogBlackmagicCore, VeryVerbose, TEXT("Scheduled frame starting with odd timecode. (%d)"), InFrame->Timecode.Frames);
+				}
+			}
+
 			const BMDTimeValue Duration = ChannelOptions.FormatInfo.FrameRateDenominator;
 			const BMDTimeScale Scale = ChannelOptions.FormatInfo.FrameRateNumerator;
 			HRESULT ScheduleVideoResult = DeckLinkOutput->ScheduleVideoFrame(static_cast<IDeckLinkVideoFrame*>(InFrame->BlackmagicFrame.Get()), OutputTime, Duration, Scale);
@@ -1381,6 +1441,45 @@ namespace BlackmagicDesign
 #endif // PLATFORM_WINDOWS
 		}
 
+		bool FOutputChannel::ShouldCaptureThisFrame(::FTimecode Timecode, uint32 FrameNumber)
+		{
+			// In interlace, don't capture for an odd frame if we don't find an even frame for it.
+			if (ChannelOptions.IsInterlaced() && Timecode.Frames % 2 != 0 && ChannelOptions.bOutputInterlaceOnEvenFrames)
+			{
+				// Odd frame should always have a matching frame.
+				// Find a frame where the even frame is filled but not the odd frame and have the same timecode (only in interlaced)
+				for (FOutputFrame* Frame : FrameReadyToWrite)
+				{
+					if (Frame->FrameIdentifier != FOutputFrame::InvalidFrameIdentifier && Frame->FrameIdentifierF2 == FOutputFrame::InvalidFrameIdentifier)
+					{
+						if (Frame->FrameIdentifier + 1 == FrameNumber)
+						{
+							return true;
+						}
+					}
+				}
+
+				if (!FrameReadyToWrite.size())
+				{
+					UE_LOG(LogBlackmagicCore, Verbose, TEXT("Skipping output because no writable frame is available."));
+				}
+				else
+				{
+					for (FOutputFrame* Frame : FrameReadyToWrite)
+					{
+						UE_LOG(LogBlackmagicCore, VeryVerbose, TEXT("\t F1: %d, F2: %d"), Frame->FrameIdentifier, Frame->FrameIdentifierF2);
+					}
+				}
+
+				// Dont interlace starting with an odd frame.
+				return false;
+			}
+			else
+			{
+				return true;
+			}
+		}
+
 
 		bool FOutputChannel::SetAudioFrameData(const FAudioSamplesDescriptor& InSamples)
 		{
@@ -1486,6 +1585,7 @@ namespace BlackmagicDesign
 				else
 				{
 					Frame->FrameIdentifierF2 = InFrame.FrameIdentifier;
+					Frame->Timecode2 = InFrame.Timecode;
 				}
 
 				return Frame;
@@ -1533,17 +1633,22 @@ namespace BlackmagicDesign
 			// Find a new empty frames
 			if (!ChannelOptions.bOutputInterlacedFieldsTimecodeNeedToMatch)
 			{
-				if (!InFrame.bEvenFrame)
-				{
-					// If the incoming frame is interlaced and we didn't find the correspondent even field
-					// do not create a new field. Drop it.
-					++LostFrameCounter;
-					if (ChannelOptions.bLogDropFrames)
+				// We don't use this path when using bOutputInterlaceOnEvenFrames, since frames that land on an odd 
+				// frame counter in the render thread may actually be even timecode frames.
+				if (!ChannelOptions.bOutputInterlaceOnEvenFrames)
+				{ 
+					if (!InFrame.bEvenFrame)
 					{
-						UE_LOG(LogBlackmagicCore, Warning,TEXT("Could not find frame. id %d, total flight %d, available %d"), InFrame.FrameIdentifier, InFlightFrames.load(), FrameReadyToWrite.size());
-					}
+						// If the incoming frame is interlaced and we didn't find the correspondent even field
+						// do not create a new field. Drop it.
+						++LostFrameCounter;
+						if (ChannelOptions.bLogDropFrames)
+						{
+							UE_LOG(LogBlackmagicCore, Warning,TEXT("Could not find frame. id %d, total flight %d, available %d"), InFrame.FrameIdentifier, InFlightFrames.load(), FrameReadyToWrite.size());
+						}
 
-					return nullptr;
+						return nullptr;
+					}
 				}
 			}
 
diff --git a/Engine/Plugins/Media/BlackmagicMedia/Source/BlackmagicCore/Private/BlackmagicOutputChannel.h b/Engine/Plugins/Media/BlackmagicMedia/Source/BlackmagicCore/Private/BlackmagicOutputChannel.h
--- a/Engine/Plugins/Media/BlackmagicMedia/Source/BlackmagicCore/Private/BlackmagicOutputChannel.h
+++ b/Engine/Plugins/Media/BlackmagicMedia/Source/BlackmagicCore/Private/BlackmagicOutputChannel.h
@@ -204,6 +204,10 @@ namespace BlackmagicDesign
 
 			bool SetVideoFrameData(const FFrameDescriptor& InFrame);
 			bool SetVideoFrameData(FFrameDescriptor_GPUDMA& InFrame);
+
+			// Return whether to capture or discard a frame based on its timecode and frame number.
+			bool ShouldCaptureThisFrame(::FTimecode Timecode, uint32 FrameNumber);
+
 			bool SetAudioFrameData(const FAudioSamplesDescriptor& InSamples);
 
 			void OnScheduledFrameCompleted(IDeckLinkVideoFrame* InCompletedFrame, BMDOutputFrameCompletionResult InResult);
@@ -313,6 +317,7 @@ namespace BlackmagicDesign
 				uint32_t NumAudioSamples;
 
 				FTimecode Timecode;
+				FTimecode Timecode2;
 
 				uint32_t CopiedVideoBufferSize;
 				uint32_t CopiedAudioBufferSize;
diff --git a/Engine/Plugins/Media/BlackmagicMedia/Source/BlackmagicCore/Public/BlackmagicLib.h b/Engine/Plugins/Media/BlackmagicMedia/Source/BlackmagicCore/Public/BlackmagicLib.h
--- a/Engine/Plugins/Media/BlackmagicMedia/Source/BlackmagicCore/Public/BlackmagicLib.h
+++ b/Engine/Plugins/Media/BlackmagicMedia/Source/BlackmagicCore/Public/BlackmagicLib.h
@@ -243,6 +243,11 @@ namespace BlackmagicDesign
 	{
 		FOutputChannelOptions();
 
+		bool IsInterlaced() const
+		{
+			return FormatInfo.FieldDominance == EFieldDominance::Interlaced;
+		}
+
 		FFormatInfo FormatInfo;
 		int32_t CallbackPriority;
 		EPixelFormat PixelFormat;
@@ -264,9 +269,11 @@ namespace BlackmagicDesign
 		bool bOutputAudio;
 		bool bOutputInterlacedFieldsTimecodeNeedToMatch;
 		bool bOutputInterlaceAsProgressive;
+		bool bOutputInterlaceOnEvenFrames;
 		bool bLogDropFrames;
 		bool bUseGPUDMA;
 		bool bScheduleInDifferentThread;
+		bool bEncodeTimecodeInTexel;
 	};
 
 	/* IInputEventCallback definition
@@ -490,5 +497,7 @@ namespace BlackmagicDesign
 	BLACKMAGICCORE_API void UnregisterOutputChannel(const FChannelInfo& InChannelInfo, FUniqueIdentifier InIdentifier, bool bCallCompleted);
 	BLACKMAGICCORE_API bool SendVideoFrameData(const FChannelInfo& InChannelInfo, const FFrameDescriptor& InFrame);
 	BLACKMAGICCORE_API bool SendVideoFrameData(const FChannelInfo& InChannelInfo, FFrameDescriptor_GPUDMA& InFrame);
+	// Return whether to capture or discard a frame based on its timecode and frame number.
+	BLACKMAGICCORE_API bool ShouldCaptureThisFrame(const FChannelInfo& InChannelInfo, ::FTimecode Timecode, uint32 FrameNumber);
 	BLACKMAGICCORE_API bool SendAudioSamples(const FChannelInfo& InChannelInfo, const FAudioSamplesDescriptor& InSamples);
 };
diff --git a/Engine/Plugins/Media/BlackmagicMedia/Source/BlackmagicMediaOutput/Private/BlackmagicMediaCapture.cpp b/Engine/Plugins/Media/BlackmagicMedia/Source/BlackmagicMediaOutput/Private/BlackmagicMediaCapture.cpp
--- a/Engine/Plugins/Media/BlackmagicMedia/Source/BlackmagicMediaOutput/Private/BlackmagicMediaCapture.cpp
+++ b/Engine/Plugins/Media/BlackmagicMedia/Source/BlackmagicMediaOutput/Private/BlackmagicMediaCapture.cpp
@@ -139,6 +139,11 @@ namespace BlackmagicMediaCaptureHelpers
 			return BlackmagicDesign::SendAudioSamples(ChannelInfo, InAudioDescriptor);
 		}
 
+		bool ShouldCaptureThisFrame(FTimecode Timecode, uint32 FrameNumber) const
+		{
+			return BlackmagicDesign::ShouldCaptureThisFrame(ChannelInfo, Timecode, FrameNumber);
+		}
+
 	private:
 
 		virtual void AddRef() override
@@ -303,34 +308,6 @@ namespace BlackmagicMediaCaptureHelpers
 			return BlackmagicDesign::EAudioBitDepth::Signed_32Bits;
 		}
 	}
-
-	void EncodeTimecodeInTexel(EBlackmagicMediaOutputPixelFormat BlackmagicMediaOutputPixelFormat, EMediaCaptureConversionOperation ConversionOperation, BlackmagicDesign::FTimecode Timecode, void* Buffer, int32 Width, int32 Height)
-	{
-		switch (BlackmagicMediaOutputPixelFormat)
-		{
-			case EBlackmagicMediaOutputPixelFormat::PF_8BIT_YUV:
-			{
-				if (ConversionOperation == EMediaCaptureConversionOperation::RGBA8_TO_YUV_8BIT)
-				{
-					FMediaIOCoreEncodeTime EncodeTime(EMediaIOCoreEncodePixelFormat::CharUYVY, Buffer, Width * 4, Width * 2, Height);
-					EncodeTime.Render(Timecode.Hours, Timecode.Minutes, Timecode.Seconds, Timecode.Frames);
-					break;
-				}
-				else
-				{
-					FMediaIOCoreEncodeTime EncodeTime(EMediaIOCoreEncodePixelFormat::CharBGRA, Buffer, Width * 4, Width, Height);
-					EncodeTime.Render(Timecode.Hours, Timecode.Minutes, Timecode.Seconds, Timecode.Frames);
-					break;
-				}
-			}
-			case EBlackmagicMediaOutputPixelFormat::PF_10BIT_YUV:
-			{
-				FMediaIOCoreEncodeTime EncodeTime(EMediaIOCoreEncodePixelFormat::YUVv210, Buffer, Width * 16, Width * 6, Height);
-				EncodeTime.Render(Timecode.Hours, Timecode.Minutes, Timecode.Seconds, Timecode.Frames);
-				break;
-			}
-		}
-	}
 }
 
 /* namespace BlackmagicMediaCaptureDevice
@@ -379,7 +356,6 @@ namespace BlackmagicMediaCaptureAnalytics
 //*****************************************************************************/
 UBlackmagicMediaCapture::UBlackmagicMediaCapture()
 	: bWaitForSyncEvent(false)
-	, bEncodeTimecodeInTexel(false)
 	, bLogDropFrame(false)
 	, BlackmagicMediaOutputPixelFormat(EBlackmagicMediaOutputPixelFormat::PF_8BIT_YUV)
 	, bSavedIgnoreTextureAlpha(false)
@@ -545,11 +521,6 @@ void UBlackmagicMediaCapture::OnFrameCapturedInternal_AnyThread(const FCaptureBa
 	{
 		BlackmagicDesign::FTimecode Timecode = BlackmagicMediaCaptureDevice::ConvertToBlackmagicTimecode(InBaseData.SourceFrameTimecode, InBaseData.SourceFrameTimecodeFramerate.AsDecimal(), FrameRate.AsDecimal());
 
-		if (bEncodeTimecodeInTexel)
-		{
-			BlackmagicMediaCaptureHelpers::EncodeTimecodeInTexel(BlackmagicMediaOutputPixelFormat, GetConversionOperation(), Timecode, InResourceData.Buffer, InResourceData.Width, InResourceData.Height);
-		}
-
 		BlackmagicDesign::FFrameDescriptor Frame;
 		Frame.VideoBuffer = reinterpret_cast<uint8_t*>(InResourceData.Buffer);
 		Frame.VideoWidth = InResourceData.Width;
@@ -643,6 +614,16 @@ void UBlackmagicMediaCapture::WaitForGPU(FRHITexture* InRHITexture)
 	}
 }
 
+bool UBlackmagicMediaCapture::ShouldCaptureThisFrame(const FCaptureBaseData& InBaseData) const
+{
+	if (EventCallback)
+	{
+		return EventCallback->ShouldCaptureThisFrame(InBaseData.SourceFrameTimecode, InBaseData.SourceFrameNumber);
+	}
+
+	return false;
+}
+
 bool UBlackmagicMediaCapture::InitBlackmagic(UBlackmagicMediaOutput* InBlackmagicMediaOutput)
 {
 	check(InBlackmagicMediaOutput);
@@ -656,7 +637,6 @@ bool UBlackmagicMediaCapture::InitBlackmagic(UBlackmagicMediaOutput* InBlackmagi
 
 	// Init general settings
 	bWaitForSyncEvent = InBlackmagicMediaOutput->bWaitForSyncEvent;
-	bEncodeTimecodeInTexel = InBlackmagicMediaOutput->bEncodeTimecodeInTexel;
 	bLogDropFrame = InBlackmagicMediaOutput->bLogDropFrame;
 	FrameRate = InBlackmagicMediaOutput->GetRequestedFrameRate();
 
@@ -693,10 +673,12 @@ bool UBlackmagicMediaCapture::InitBlackmagic(UBlackmagicMediaOutput* InBlackmagi
 	ChannelOptions.bOutputVideo = true;
 	ChannelOptions.bOutputInterlacedFieldsTimecodeNeedToMatch = InBlackmagicMediaOutput->bInterlacedFieldsTimecodeNeedToMatch && InBlackmagicMediaOutput->OutputConfiguration.MediaConfiguration.MediaMode.Standard == EMediaIOStandardType::Interlaced && InBlackmagicMediaOutput->TimecodeFormat != EMediaIOTimecodeFormat::None;
 	ChannelOptions.bOutputInterlaceAsProgressive = InBlackmagicMediaOutput->bOutputInterlaceAsProgressive;
+	ChannelOptions.bOutputInterlaceOnEvenFrames = InBlackmagicMediaOutput->bOutputInterlaceOnEvenFrames;
 	ChannelOptions.bLogDropFrames = bLogDropFrame;
 	ChannelOptions.bUseGPUDMA = ShouldCaptureRHIResource();
 	ChannelOptions.bScheduleInDifferentThread = InBlackmagicMediaOutput->bUseMultithreadedScheduling;
 	ChannelOptions.HDRMetadata = BlackmagicMediaCaptureHelpers::MakeBlackmagicHDRMetadata(InBlackmagicMediaOutput->HDROptions);
+	ChannelOptions.bEncodeTimecodeInTexel = InBlackmagicMediaOutput->bEncodeTimecodeInTexel;
 
 	AudioBitDepth = InBlackmagicMediaOutput->AudioBitDepth;
 	bOutputAudio = InBlackmagicMediaOutput->bOutputAudio;
diff --git a/Engine/Plugins/Media/BlackmagicMedia/Source/BlackmagicMediaOutput/Public/BlackmagicMediaCapture.h b/Engine/Plugins/Media/BlackmagicMedia/Source/BlackmagicMediaOutput/Public/BlackmagicMediaCapture.h
--- a/Engine/Plugins/Media/BlackmagicMedia/Source/BlackmagicMediaOutput/Public/BlackmagicMediaCapture.h
+++ b/Engine/Plugins/Media/BlackmagicMedia/Source/BlackmagicMediaOutput/Public/BlackmagicMediaCapture.h
@@ -60,7 +60,7 @@ protected:
 	virtual bool SupportsAnyThreadCapture() const override;
 	virtual const FMatrix& GetRGBToYUVConversionMatrix() const override;
 	virtual void WaitForGPU(FRHITexture* InRHITexture) override;
-
+	virtual bool ShouldCaptureThisFrame(const FCaptureBaseData& InBaseData) const override;
 
 private:
 	struct FBlackmagicOutputCallback;
@@ -81,7 +81,6 @@ private:
 
 	/** Option from MediaOutput */
 	bool bWaitForSyncEvent;
-	bool bEncodeTimecodeInTexel;
 	bool bLogDropFrame;
 	
 	/** MediaOutput cached value */
diff --git a/Engine/Plugins/Media/BlackmagicMedia/Source/BlackmagicMediaOutput/Public/BlackmagicMediaOutput.h b/Engine/Plugins/Media/BlackmagicMedia/Source/BlackmagicMediaOutput/Public/BlackmagicMediaOutput.h
--- a/Engine/Plugins/Media/BlackmagicMedia/Source/BlackmagicMediaOutput/Public/BlackmagicMediaOutput.h
+++ b/Engine/Plugins/Media/BlackmagicMedia/Source/BlackmagicMediaOutput/Public/BlackmagicMediaOutput.h
@@ -109,6 +109,13 @@ public:
 	 */
 	UPROPERTY(EditAnywhere, AdvancedDisplay, Category = "Output")
 	bool bOutputInterlaceAsProgressive = false;
+
+	/**
+	 * Whether to only schedule interlace frames if we land on an even timecode. Useful when you have an interlace input
+	 * and want to ensure that we interlace the original frames. (Experimental)
+	 */
+	UPROPERTY(EditAnywhere, AdvancedDisplay, Category = "Output")
+	bool bOutputInterlaceOnEvenFrames = false;
 	
 	/**
 	 * HDR Metadata of the video signal.
