Media IO - Allow specifying a pixel offset to FMediaIOCoreEncodeTime

diff --git a/Engine/Plugins/Media/BlackmagicMedia/Source/BlackmagicCore/Private/BlackmagicOutputChannel.cpp b/Engine/Plugins/Media/BlackmagicMedia/Source/BlackmagicCore/Private/BlackmagicOutputChannel.cpp
--- a/Engine/Plugins/Media/BlackmagicMedia/Source/BlackmagicCore/Private/BlackmagicOutputChannel.cpp
+++ b/Engine/Plugins/Media/BlackmagicMedia/Source/BlackmagicCore/Private/BlackmagicOutputChannel.cpp
@@ -108,28 +108,30 @@ namespace BlackmagicDesign
 			return RefCount;
 		}
 
-		void EncodeTimecodeInTexel(BMDPixelFormat PixelFormat, BlackmagicDesign::FTimecode Timecode, void* Buffer, int32 Width, int32 Height)
+		void EncodeTimecodeInTexel(BMDPixelFormat PixelFormat, BlackmagicDesign::FTimecode Timecode, void* Buffer, int32 Width, int32 Height, int32 XOffset = 0, int32 YOffset = 0)
 		{
 			switch (PixelFormat)
 			{
 			case bmdFormat8BitYUV:
 			{
-				FMediaIOCoreEncodeTime EncodeTime(EMediaIOCoreEncodePixelFormat::CharUYVY, Buffer, Width * 4, Width * 2, Height);
+				FMediaIOCoreEncodeTime EncodeTime(EMediaIOCoreEncodePixelFormat::CharUYVY, Buffer, Width * 4, Width * 2, Height, XOffset, YOffset);
 				EncodeTime.Render(Timecode.Hours, Timecode.Minutes, Timecode.Seconds, Timecode.Frames);
 				break;
 			}
 			case bmdFormat8BitBGRA:
 			{
-				FMediaIOCoreEncodeTime EncodeTime(EMediaIOCoreEncodePixelFormat::CharBGRA, Buffer, Width * 4, Width, Height);
+				FMediaIOCoreEncodeTime EncodeTime(EMediaIOCoreEncodePixelFormat::CharBGRA, Buffer, Width * 4, Width, Height, XOffset, YOffset);
 				EncodeTime.Render(Timecode.Hours, Timecode.Minutes, Timecode.Seconds, Timecode.Frames);
 				break;
 			}
 			case bmdFormat10BitYUV:
 			{
-				FMediaIOCoreEncodeTime EncodeTime(EMediaIOCoreEncodePixelFormat::YUVv210, Buffer, Width * 16, Width * 6, Height);
+				FMediaIOCoreEncodeTime EncodeTime(EMediaIOCoreEncodePixelFormat::YUVv210, Buffer, Width * 16, Width * 6, Height, XOffset, YOffset);
 				EncodeTime.Render(Timecode.Hours, Timecode.Minutes, Timecode.Seconds, Timecode.Frames);
 				break;
 			}
+			default:
+				break;
 			}
 		}
 
@@ -648,7 +650,7 @@ namespace BlackmagicDesign
 				if (ChannelOptions.bEncodeTimecodeInTexel)
 				{
 					EncodeTimecodeInTexel(PixelFormat, InFrame->Timecode, DestinationBuffer, VideoWidth, VideoHeight);
-					EncodeTimecodeInTexel(PixelFormat, InFrame->Timecode, (void*)((uint8*)DestinationBuffer + 200 * 4), VideoWidth, VideoHeight);
+					EncodeTimecodeInTexel(PixelFormat, InFrame->Timecode2, DestinationBuffer, VideoWidth, VideoHeight, 0, 20);
 				}
 
 				if (InFrame->Timecode.Frames % 2 == 0)
diff --git a/Engine/Plugins/Media/MediaIOFramework/Source/MediaIOCore/Private/Shared/MediaIOCoreEncodeTime.cpp b/Engine/Plugins/Media/MediaIOFramework/Source/MediaIOCore/Private/Shared/MediaIOCoreEncodeTime.cpp
--- a/Engine/Plugins/Media/MediaIOFramework/Source/MediaIOCore/Private/Shared/MediaIOCoreEncodeTime.cpp
+++ b/Engine/Plugins/Media/MediaIOFramework/Source/MediaIOCore/Private/Shared/MediaIOCoreEncodeTime.cpp
@@ -2,12 +2,14 @@
 
 #include "MediaIOCoreEncodeTime.h"
 
-FMediaIOCoreEncodeTime::FMediaIOCoreEncodeTime(EMediaIOCoreEncodePixelFormat InFormat, void* InBuffer, uint32 InPitch, uint32 InWidth, uint32 InHeight)
+FMediaIOCoreEncodeTime::FMediaIOCoreEncodeTime(EMediaIOCoreEncodePixelFormat InFormat, void* InBuffer, uint32 InPitch, uint32 InWidth, uint32 InHeight, uint32 XOffset, uint32 YOffset)
 	: Format(InFormat)
 	, Buffer(InBuffer)
 	, Pitch(InPitch)
 	, Width(InWidth)
 	, Height(InHeight)
+	, XPixelOffset(XOffset)
+	, YPixelOffset(YOffset)
 {
 	check(Buffer);
 
@@ -62,7 +64,7 @@ void FMediaIOCoreEncodeTime::SetPixelScaled(uint32 InX, uint32 InY, bool InSet,
 	{
 		for (uint32 ScaleX = 0; ScaleX < InScale; ScaleX++)
 		{
-			SetPixel(InX * InScale + ScaleX, InY * InScale + ScaleY, InSet);
+			SetPixel((InX+XPixelOffset) * InScale + ScaleX, (InY+YPixelOffset) * InScale + ScaleY, InSet);
 		}
 	}
 }
diff --git a/Engine/Plugins/Media/MediaIOFramework/Source/MediaIOCore/Public/MediaIOCoreEncodeTime.h b/Engine/Plugins/Media/MediaIOFramework/Source/MediaIOCore/Public/MediaIOCoreEncodeTime.h
--- a/Engine/Plugins/Media/MediaIOFramework/Source/MediaIOCore/Public/MediaIOCoreEncodeTime.h
+++ b/Engine/Plugins/Media/MediaIOFramework/Source/MediaIOCore/Public/MediaIOCoreEncodeTime.h
@@ -18,7 +18,7 @@ enum class EMediaIOCoreEncodePixelFormat
 class FMediaIOCoreEncodeTime
 {
 public:
-	UE_API FMediaIOCoreEncodeTime(EMediaIOCoreEncodePixelFormat InFormat, void* InBuffer, uint32 InPitch, uint32 InWidth, uint32 InHeight);
+	UE_API FMediaIOCoreEncodeTime(EMediaIOCoreEncodePixelFormat InFormat, void* InBuffer, uint32 InPitch, uint32 InWidth, uint32 InHeight, uint32 XOffset = 0, uint32 YOffset = 0);
 	UE_API void Render(uint32 InHours, uint32 InMinutes, uint32 InSeconds, uint32 InFrames) const;
 
 protected:
@@ -49,6 +49,12 @@ protected:
 	/** Height of image */
 	uint32 Height;
 
+	/** X pixel offset of the rendered timecode. */
+	uint32 XPixelOffset;
+
+	/** Y pixel offset of the rendered timecode. */
+	uint32 YPixelOffset;
+
 protected:
 	TColor ColorBlack;
 	TColor ColorWhite;
