Disable BlackmagicMedia plugin on Windows Arm64

#rnx

diff --git a/Engine/Plugins/Media/BlackmagicMedia/BlackmagicMedia.uplugin b/Engine/Plugins/Media/BlackmagicMedia/BlackmagicMedia.uplugin
--- a/Engine/Plugins/Media/BlackmagicMedia/BlackmagicMedia.uplugin
+++ b/Engine/Plugins/Media/BlackmagicMedia/BlackmagicMedia.uplugin
@@ -18,37 +18,43 @@
 			"Name": "BlackmagicCore",
 			"Type": "Runtime",
 			"LoadingPhase": "PreLoadingScreen",
-			"PlatformAllowList": [ "Win64", "Linux" ]
+			"PlatformAllowList": [ "Win64", "Linux" ],
+			"PlatformArchitectureDenyList": [ "Win64:arm64" ]
 		},
 		{
 			"Name": "BlackmagicMedia",
 			"Type": "Runtime",
 			"LoadingPhase": "PreLoadingScreen",
-			"PlatformAllowList": [ "Win64", "Linux" ]
+			"PlatformAllowList": [ "Win64", "Linux" ],
+			"PlatformArchitectureDenyList": [ "Win64:arm64" ]
 		},
 		{
 			"Name": "BlackmagicMediaOutput",
 			"Type": "Runtime",
 			"LoadingPhase": "PreLoadingScreen",
-			"PlatformAllowList": [ "Win64", "Linux" ]
+			"PlatformAllowList": [ "Win64", "Linux" ],
+			"PlatformArchitectureDenyList": [ "Win64:arm64" ]
 		},
 		{
 			"Name": "BlackmagicMediaFactory",
 			"Type": "Editor",
 			"LoadingPhase": "PostEngineInit",
-			"PlatformAllowList": [ "Win64", "Linux" ]
+			"PlatformAllowList": [ "Win64", "Linux" ],
+			"PlatformArchitectureDenyList": [ "Win64:arm64" ]
 		},
 		{
 			"Name": "BlackmagicMediaFactory",
 			"Type": "RuntimeNoCommandlet",
 			"LoadingPhase": "PostEngineInit",
-			"PlatformAllowList": [ "Win64", "Linux" ]
+			"PlatformAllowList": [ "Win64", "Linux" ],
+			"PlatformArchitectureDenyList": [ "Win64:arm64" ]
 		},
 		{
 			"Name": "BlackmagicMediaEditor",
 			"Type": "Editor",
 			"LoadingPhase": "PostEngineInit",
-			"PlatformAllowList": [ "Win64", "Linux" ]
+			"PlatformAllowList": [ "Win64", "Linux" ],
+			"PlatformArchitectureDenyList": [ "Win64:arm64" ]
 		}
 	],
 	"Plugins": [
