Add missing include
#rnx

diff --git a/Engine/Plugins/Media/BlackmagicMedia/Source/BlackmagicCore/Private/BlackmagicHelper.cpp b/Engine/Plugins/Media/BlackmagicMedia/Source/BlackmagicCore/Private/BlackmagicHelper.cpp
--- a/Engine/Plugins/Media/BlackmagicMedia/Source/BlackmagicCore/Private/BlackmagicHelper.cpp
+++ b/Engine/Plugins/Media/BlackmagicMedia/Source/BlackmagicCore/Private/BlackmagicHelper.cpp
@@ -3,6 +3,7 @@
 #include "BlackmagicHelper.h"
 
 #include "Common.h"
+#include <chrono>
 
 /*
  * Global helper functions
