Blackmagic - Fix Blackmagic source assigning default configuration despite having a valid one.

diff --git a/Engine/Plugins/Media/BlackmagicMedia/Source/BlackmagicMedia/Private/Assets/BlackmagicMediaSource.cpp b/Engine/Plugins/Media/BlackmagicMedia/Source/BlackmagicMedia/Private/Assets/BlackmagicMediaSource.cpp
--- a/Engine/Plugins/Media/BlackmagicMedia/Source/BlackmagicMedia/Private/Assets/BlackmagicMediaSource.cpp
+++ b/Engine/Plugins/Media/BlackmagicMedia/Source/BlackmagicMedia/Private/Assets/BlackmagicMediaSource.cpp
@@ -20,7 +20,6 @@ UBlackmagicMediaSource::UBlackmagicMediaSource()
 	, bLogDropFrame(false)
 	, bEncodeTimecodeInTexel(false)
 {
-	AssignDefaultConfiguration();
 }
 
 /*
@@ -293,8 +292,6 @@ void UBlackmagicMediaSource::PostLoad()
 {
 	Super::PostLoad();
 
-	AssignDefaultConfiguration();
-
 #if WITH_EDITORONLY_DATA
 	PRAGMA_DISABLE_DEPRECATION_WARNINGS
 		if (TimecodeFormat_DEPRECATED != EMediaIOTimecodeFormat::None)
diff --git a/Engine/Plugins/Media/BlackmagicMedia/Source/BlackmagicMedia/Public/BlackmagicMediaSource.h b/Engine/Plugins/Media/BlackmagicMedia/Source/BlackmagicMedia/Public/BlackmagicMediaSource.h
--- a/Engine/Plugins/Media/BlackmagicMedia/Source/BlackmagicMedia/Public/BlackmagicMediaSource.h
+++ b/Engine/Plugins/Media/BlackmagicMedia/Source/BlackmagicMedia/Public/BlackmagicMediaSource.h
@@ -129,6 +129,9 @@ public:
 	virtual void GetDetailsPanelInfoElements(TArray<FInfoElement>& OutInfoElements) const override;
 #endif //WITH_EDITOR
 	
+	/** Assigns a default configuration (The first one in the list of configurations). */
+	void AssignDefaultConfiguration();
+
 public:
 	//~ UObject interface
 #if WITH_EDITOR
@@ -137,8 +140,4 @@ public:
 #endif //WITH_EDITOR
 	virtual void PostLoad() override;
 	//~ End UObject interface
-
-private:
-	/** Assigns a default configuration if the current configuration is invalid. */
-	void AssignDefaultConfiguration();
 };
diff --git a/Engine/Plugins/Media/BlackmagicMedia/Source/BlackmagicMediaEditor/Private/Factories/BlackmagicMediaSourceFactoryNew.cpp b/Engine/Plugins/Media/BlackmagicMedia/Source/BlackmagicMediaEditor/Private/Factories/BlackmagicMediaSourceFactoryNew.cpp
--- a/Engine/Plugins/Media/BlackmagicMedia/Source/BlackmagicMediaEditor/Private/Factories/BlackmagicMediaSourceFactoryNew.cpp
+++ b/Engine/Plugins/Media/BlackmagicMedia/Source/BlackmagicMediaEditor/Private/Factories/BlackmagicMediaSourceFactoryNew.cpp
@@ -23,7 +23,9 @@ UBlackmagicMediaSourceFactoryNew::UBlackmagicMediaSourceFactoryNew(const FObject
 
 UObject* UBlackmagicMediaSourceFactoryNew::FactoryCreateNew(UClass* InClass, UObject* InParent, FName InName, EObjectFlags Flags, UObject* Context, FFeedbackContext* Warn)
 {
-	return NewObject<UBlackmagicMediaSource>(InParent, InClass, InName, Flags);
+	UBlackmagicMediaSource* Source = NewObject<UBlackmagicMediaSource>(InParent, InClass, InName, Flags);
+	Source->AssignDefaultConfiguration();
+	return Source;
 }
 
 
