Blackmagic - Fix Assertion caused by accessing FApp::GetCurrentTime from outside the game thread

#Synced-CL 48456341

diff --git a/Engine/Plugins/Media/BlackmagicMedia/Source/BlackmagicMedia/Private/Player/BlackmagicMediaPlayer.cpp b/Engine/Plugins/Media/BlackmagicMedia/Source/BlackmagicMedia/Private/Player/BlackmagicMediaPlayer.cpp
--- a/Engine/Plugins/Media/BlackmagicMedia/Source/BlackmagicMedia/Private/Player/BlackmagicMediaPlayer.cpp
+++ b/Engine/Plugins/Media/BlackmagicMedia/Source/BlackmagicMedia/Private/Player/BlackmagicMediaPlayer.cpp
@@ -181,7 +181,7 @@ namespace BlackmagicMediaPlayerHelpers
 
 			if (!InFrameInfo.bHasInputSource && InFrameInfo.AudioBuffer == nullptr)
 			{
-				const double CurrentTime = FApp::GetCurrentTime();
+				const double CurrentTime = FPlatformTime::Seconds();
 				const double TimeAllowedToConnect = 2.0;
 				if (LastHasFrameTime < 0.1)
 				{
